Adds Base for Uncontrolled Changelist module.

#changelist validated

diff --git a/Engine/Plugins/Developer/PerforceSourceControl/Source/PerforceSourceControl/Private/PerforceSourceControlChangelistState.cpp b/Engine/Plugins/Developer/PerforceSourceControl/Source/PerforceSourceControl/Private/PerforceSourceControlChangelistState.cpp
--- a/Engine/Plugins/Developer/PerforceSourceControl/Source/PerforceSourceControl/Private/PerforceSourceControlChangelistState.cpp
+++ b/Engine/Plugins/Developer/PerforceSourceControl/Source/PerforceSourceControl/Private/PerforceSourceControlChangelistState.cpp
@@ -12,12 +12,12 @@ FPerforceSourceControlChangelistState::FPerforceSourceControlChangelistState(con
 
 FName FPerforceSourceControlChangelistState::GetIconName() const
 {
-	return FName("Perforce.CheckedOut");
+	return FName("SourceControl.Changelist");
 }
 
 FName FPerforceSourceControlChangelistState::GetSmallIconName() const
 {
-	return FName("Perforce.CheckedOut_Small");
+	return FName("SourceControl.Changelist");
 }
 
 FText FPerforceSourceControlChangelistState::GetDisplayText() const
diff --git a/Engine/Plugins/Developer/PerforceSourceControl/Source/PerforceSourceControl/Private/PerforceSourceControlModule.cpp b/Engine/Plugins/Developer/PerforceSourceControl/Source/PerforceSourceControl/Private/PerforceSourceControlModule.cpp
--- a/Engine/Plugins/Developer/PerforceSourceControl/Source/PerforceSourceControl/Private/PerforceSourceControlModule.cpp
+++ b/Engine/Plugins/Developer/PerforceSourceControl/Source/PerforceSourceControl/Private/PerforceSourceControlModule.cpp
@@ -27,7 +27,8 @@ void FPerforceSourceControlModule::StartupModule()
 	PerforceSourceControlProvider.RegisterWorker( "GetWorkspaces", FGetPerforceSourceControlWorker::CreateStatic( &CreateWorker<FPerforceGetWorkspacesWorker> ) );
 	PerforceSourceControlProvider.RegisterWorker( "Copy", FGetPerforceSourceControlWorker::CreateStatic( &CreateWorker<FPerforceCopyWorker> ) );
 	PerforceSourceControlProvider.RegisterWorker( "Resolve", FGetPerforceSourceControlWorker::CreateStatic( &CreateWorker<FPerforceResolveWorker> ) );
-	PerforceSourceControlProvider.RegisterWorker( "ChangeStatus", FGetPerforceSourceControlWorker::CreateStatic( &CreateWorker<FPerforceChangeStatusWorker> ) );
+	PerforceSourceControlProvider.RegisterWorker( "Status", FGetPerforceSourceControlWorker::CreateStatic(&CreateWorker<FPerforceStatusWorker>));
+	PerforceSourceControlProvider.RegisterWorker( "ChangeStatus", FGetPerforceSourceControlWorker::CreateStatic(&CreateWorker<FPerforceChangeStatusWorker>));
 	PerforceSourceControlProvider.RegisterWorker( "UpdateChangelistsStatus", FGetPerforceSourceControlWorker::CreateStatic( &CreateWorker<FPerforceGetPendingChangelistsWorker> ) );
 	PerforceSourceControlProvider.RegisterWorker( "NewChangelist", FGetPerforceSourceControlWorker::CreateStatic( &CreateWorker<FPerforceNewChangelistWorker> ) );
 	PerforceSourceControlProvider.RegisterWorker( "DeleteChangelist", FGetPerforceSourceControlWorker::CreateStatic( &CreateWorker<FPerforceDeleteChangelistWorker> ) );
diff --git a/Engine/Plugins/Developer/PerforceSourceControl/Source/PerforceSourceControl/Private/PerforceSourceControlOperations.cpp b/Engine/Plugins/Developer/PerforceSourceControl/Source/PerforceSourceControl/Private/PerforceSourceControlOperations.cpp
--- a/Engine/Plugins/Developer/PerforceSourceControl/Source/PerforceSourceControl/Private/PerforceSourceControlOperations.cpp
+++ b/Engine/Plugins/Developer/PerforceSourceControl/Source/PerforceSourceControl/Private/PerforceSourceControlOperations.cpp
@@ -6,6 +6,7 @@
 #include "Misc/Paths.h"
 #include "Misc/EngineVersion.h"
 #include "Modules/ModuleManager.h"
+#include "SourceControlHelpers.h"
 #include "SourceControlOperations.h"
 #include "PerforceSourceControlRevision.h"
 #include "PerforceSourceControlCommand.h"
@@ -175,6 +176,10 @@ static void ParseRecordSetForState(const FP4RecordSet& InRecords, TMap<FString,
 			{
 				OutResults.Add(FullPath, EPerforceState::Branched);
 			}
+			else if (Action == TEXT("cleared"))
+			{
+				OutResults.Add(FullPath, EPerforceState::ReadOnly);
+			}
 		}
 	}
 }
@@ -659,6 +664,11 @@ bool FPerforceRevertWorker::Execute(FPerforceSourceControlCommand& InCommand)
 			AppendChangelistParameter(Parameters);
 		}
 
+		if (StaticCastSharedRef<FRevert>(InCommand.Operation)->IsSoftRevert())
+		{
+			Parameters.Add(TEXT("-k"));
+		}
+
 		if (InCommand.Files.Num() != 0)
 		{
 			Parameters.Append(InCommand.Files);
@@ -2057,6 +2067,89 @@ bool FPerforceResolveWorker::UpdateStates() const
 	return UpdatedFiles.Num() > 0;
 }
 
+FName FPerforceStatusWorker::GetName() const
+{
+	return "Status";
+}
+
+bool FPerforceStatusWorker::Execute(class FPerforceSourceControlCommand& InCommand)
+{
+	FScopedPerforceConnection ScopedConnection(InCommand);
+	
+	if (!InCommand.IsCanceled() && ScopedConnection.IsValid() && (!InCommand.Files.IsEmpty()))
+	{
+		FPerforceConnection& Connection = ScopedConnection.GetConnection();
+
+		TArray<FString> Parameters;
+		Parameters.Append(InCommand.Files);
+
+		RequestedFileStatus = SourceControlHelpers::AbsoluteFilenames(InCommand.Files);
+
+		FP4RecordSet Records;
+		InCommand.bCommandSuccessful = Connection.RunCommand(TEXT("status"), Parameters, Records, InCommand.ResultInfo.ErrorMessages, FOnIsCancelled::CreateRaw(&InCommand, &FPerforceSourceControlCommand::IsCanceled), InCommand.bConnectionDropped);
+		
+		if (InCommand.bCommandSuccessful)
+		{
+			for (const FP4Record& Record : Records)
+			{
+				FString Filename = Record(TEXT("clientFile"));
+				FString Status = Record(TEXT("action"));
+
+				// Here we use the same method as PerforceSourceControlProvider::GetState to be sure to not create duplicates
+				TArray<FString> AbsoluteFilenames = SourceControlHelpers::AbsoluteFilenames({ Filename });
+
+				if (AbsoluteFilenames.IsEmpty() || Status.IsEmpty())
+					continue;
+
+				if (Status.Equals(TEXT("add")))
+				{
+					UpdatedFileStatus.Emplace(AbsoluteFilenames[0], EAction::Add);
+				}
+				else if (Status.Equals(TEXT("edit")))
+				{
+					UpdatedFileStatus.Emplace(AbsoluteFilenames[0], EAction::Edit);
+				}
+				else if (Status.Equals(TEXT("delete")))
+				{
+					UpdatedFileStatus.Emplace(AbsoluteFilenames[0], EAction::Delete);
+				}
+			}
+		}
+	}
+
+	return InCommand.bCommandSuccessful;
+}
+
+bool FPerforceStatusWorker::UpdateStates() const
+{
+	FPerforceSourceControlModule& PerforceSourceControl = FPerforceSourceControlModule::Get();
+	
+	for (const FString& Filename : RequestedFileStatus)
+	{
+		const FFilenameStatusCache::ValueType* StatusPtr = UpdatedFileStatus.Find(Filename);
+		TSharedRef<FPerforceSourceControlState, ESPMode::ThreadSafe> FileState = PerforceSourceControl.GetProvider().GetStateInternal(Filename);
+
+		if ((StatusPtr == nullptr) || (*StatusPtr == EAction::None))
+		{
+			FileState->bModifed = false;
+		}
+		else if (*StatusPtr == EAction::Edit)
+		{
+			FileState->bModifed = true;
+		}
+		else if (*StatusPtr == EAction::Add)
+		{
+			
+		}
+		else if (*StatusPtr == EAction::Delete)
+		{
+
+		}
+	}
+
+	return true;
+}
+
 FName FPerforceChangeStatusWorker::GetName() const
 {
 	return "ChangeStatus";
diff --git a/Engine/Plugins/Developer/PerforceSourceControl/Source/PerforceSourceControl/Private/PerforceSourceControlOperations.h b/Engine/Plugins/Developer/PerforceSourceControl/Source/PerforceSourceControl/Private/PerforceSourceControlOperations.h
--- a/Engine/Plugins/Developer/PerforceSourceControl/Source/PerforceSourceControl/Private/PerforceSourceControlOperations.h
+++ b/Engine/Plugins/Developer/PerforceSourceControl/Source/PerforceSourceControl/Private/PerforceSourceControlOperations.h
@@ -196,6 +196,31 @@ private:
 	TArray< FString > UpdatedFiles;
 };
 
+class FPerforceStatusWorker : public IPerforceSourceControlWorker
+{
+public:
+	enum class EAction
+	{
+		None = 0,
+		Add,
+		Edit,
+		Delete
+	};
+
+	typedef TMap<FString, EAction> FFilenameStatusCache;
+
+public:
+	virtual ~FPerforceStatusWorker() {}
+	// IPerforceSourceControlWorker interface
+	virtual FName GetName() const override;
+	virtual bool Execute(class FPerforceSourceControlCommand& InCommand) override;
+	virtual bool UpdateStates() const override;
+
+private:
+	TArray<FString> RequestedFileStatus;
+	FFilenameStatusCache UpdatedFileStatus;
+};
+
 class FPerforceChangeStatusWorker : public IPerforceSourceControlWorker
 {
 public:
diff --git a/Engine/Source/Developer/SourceControl/Public/SourceControlOperations.h b/Engine/Source/Developer/SourceControl/Public/SourceControlOperations.h
--- a/Engine/Source/Developer/SourceControl/Public/SourceControlOperations.h
+++ b/Engine/Source/Developer/SourceControl/Public/SourceControlOperations.h
@@ -168,6 +168,19 @@ public:
 	{
 		return LOCTEXT("SourceControl_Revert", "Reverting file(s) in Source Control...");
 	}
+
+	void SetSoftRevert(const bool bInSoftRevert)
+	{
+		bIsSoftRevert = bInSoftRevert;
+	}
+
+	bool IsSoftRevert() const
+	{
+		return bIsSoftRevert;
+	}
+
+protected:
+	bool	bIsSoftRevert = false;
 };
 
 /**
@@ -617,4 +630,22 @@ public:
 	}
 };
 
+/**
+ * Operation used to get the status of files
+ */
+class FStatus : public FSourceControlOperationBase
+{
+public:
+	// ISourceControlOperation interface
+	virtual FName GetName() const override
+	{
+		return "Status";
+	}
+
+	virtual FText GetInProgressString() const override
+	{
+		return LOCTEXT("SourceControl_StatusOperation", "Retrieving files status...");
+	}
+};
+
 #undef LOCTEXT_NAMESPACE
diff --git a/Engine/Source/Developer/UncontrolledChangelists/Private/UncontrolledChangelist.cpp b/Engine/Source/Developer/UncontrolledChangelists/Private/UncontrolledChangelist.cpp
new file mode 100644
--- /dev/null
+++ b/Engine/Source/Developer/UncontrolledChangelists/Private/UncontrolledChangelist.cpp
@@ -0,0 +1,57 @@
+// Copyright Epic Games, Inc. All Rights Reserved.
+
+#include "UncontrolledChangelist.h"
+
+#include "Dom/JsonObject.h"
+#include "ISourceControlModule.h"
+#include "Misc/Guid.h"
+
+#define LOCTEXT_NAMESPACE "UncontrolledChangelists"
+
+const FGuid FUncontrolledChangelist::DEFAULT_UNCONTROLLED_CHANGELIST_GUID = FGuid(0, 0, 0, 0);
+const FText FUncontrolledChangelist::DEFAULT_UNCONTROLLED_CHANGELIST_NAME = LOCTEXT("DefaultUncontrolledChangelist", "Default Uncontrolled Changelist");
+
+FUncontrolledChangelist::FUncontrolledChangelist()
+	: Guid(FGuid::NewGuid())
+	, Name()
+{
+}
+
+FUncontrolledChangelist::FUncontrolledChangelist(const FGuid& InGuid, const FString& InName)
+	: Guid(InGuid)
+	, Name(InName)
+{
+}
+
+void FUncontrolledChangelist::Serialize(TSharedRef<FJsonObject> OutJsonObject) const
+{
+	OutJsonObject->SetStringField(GUID_NAME, Guid.ToString());
+	OutJsonObject->SetStringField(NAME_NAME, Name);
+}
+
+bool FUncontrolledChangelist::Deserialize(const TSharedRef<FJsonObject> InJsonValue)
+{
+	FString TempStr;
+
+	if (!InJsonValue->TryGetStringField(GUID_NAME, TempStr))
+	{
+		UE_LOG(LogSourceControl, Error, TEXT("Cannot get field %s."), GUID_NAME);
+		return false;
+	}
+
+	if (!FGuid::Parse(TempStr, Guid))
+	{
+		UE_LOG(LogSourceControl, Error, TEXT("Cannot parse Guid %s."), *TempStr);
+		return false;
+	}
+
+	if (!InJsonValue->TryGetStringField(NAME_NAME, Name))
+	{
+		UE_LOG(LogSourceControl, Error, TEXT("Cannot get field %s."), NAME_NAME);
+		return false;
+	}
+
+	return true;
+}
+
+#undef LOCTEXT_NAMESPACE
diff --git a/Engine/Source/Developer/UncontrolledChangelists/Private/UncontrolledChangelistState.cpp b/Engine/Source/Developer/UncontrolledChangelists/Private/UncontrolledChangelistState.cpp
new file mode 100644
--- /dev/null
+++ b/Engine/Source/Developer/UncontrolledChangelists/Private/UncontrolledChangelistState.cpp
@@ -0,0 +1,151 @@
+// Copyright Epic Games, Inc. All Rights Reserved.
+
+#include "UncontrolledChangelistState.h"
+
+#include "Algo/Transform.h"
+#include "Dom/JsonObject.h"
+#include "Dom/JSonValue.h"
+#include "ISourceControlModule.h"
+#include "ISourceControlProvider.h"
+#include "SourceControlOperations.h"
+#include "UncontrolledChangelist.h"
+
+#define LOCTEXT_NAMESPACE "UncontrolledChangelists"
+
+FUncontrolledChangelistState::FUncontrolledChangelistState(const FUncontrolledChangelist& InUncontrolledChangelist)
+	: Changelist(InUncontrolledChangelist)
+{
+}
+
+FName FUncontrolledChangelistState::GetIconName() const
+{
+	return FName("SourceControl.UncontrolledChangelist");
+}
+
+FName FUncontrolledChangelistState::GetSmallIconName() const
+{
+	return FName("SourceControl.UncontrolledChangelist_Small");
+}
+
+FText FUncontrolledChangelistState::GetDisplayText() const
+{
+	return FText::FromString(Changelist.ToString());
+}
+
+FText FUncontrolledChangelistState::GetDescriptionText() const
+{
+	return FText::FromString(Description);
+}
+
+FText FUncontrolledChangelistState::GetDisplayTooltip() const
+{
+	return LOCTEXT("Tooltip", "Tooltip");
+}
+
+const FDateTime& FUncontrolledChangelistState::GetTimeStamp() const
+{
+	return TimeStamp;
+}
+
+const TSet<FSourceControlStateRef>& FUncontrolledChangelistState::GetFilesStates() const
+{
+	return Files;
+}
+
+void FUncontrolledChangelistState::Serialize(TSharedRef<FJsonObject> OutJsonObject) const
+{
+	TArray<TSharedPtr<FJsonValue>> FileValues;
+
+	for (const FSourceControlStateRef& File : Files)
+	{
+		FileValues.Add(MakeShareable(new FJsonValueString(File->GetFilename())));
+	}
+
+	OutJsonObject->SetArrayField(FILES_NAME, MoveTemp(FileValues));
+}
+
+bool FUncontrolledChangelistState::Deserialize(const TSharedRef<FJsonObject> InJsonValue)
+{
+	const TArray<TSharedPtr<FJsonValue>>* FileValues = nullptr;
+	ISourceControlProvider& SourceControlProvider = ISourceControlModule::Get().GetProvider();
+
+	if ((!InJsonValue->TryGetArrayField(FILES_NAME, FileValues)) || (FileValues == nullptr))
+	{
+		UE_LOG(LogSourceControl, Error, TEXT("Cannot get field %s."), FILES_NAME);
+		return false;
+	}
+
+	TArray<FString> Filenames;
+
+	Algo::Transform(*FileValues, Filenames, [](const TSharedPtr<FJsonValue>& File)
+	{
+		return File->AsString();
+	});
+
+	AddFiles(Filenames, ECheckFlags::Modified | ECheckFlags::NotCheckedOut);
+
+	return true;
+}
+
+void FUncontrolledChangelistState::AddFiles(const TArray<FString>& InFilenames, const ECheckFlags InCheckFlags)
+{
+	TArray<FSourceControlStateRef> FileStates;
+	ISourceControlProvider& SourceControlProvider = ISourceControlModule::Get().GetProvider();
+	bool bCheckStatus = (InCheckFlags & ECheckFlags::Modified) != ECheckFlags::None;
+	bool bCheckCheckout = (InCheckFlags & ECheckFlags::NotCheckedOut) != ECheckFlags::None;
+
+	if (InFilenames.IsEmpty())
+	{
+		return;
+	}
+
+	if (bCheckStatus)
+	{
+		SourceControlProvider.Execute(ISourceControlOperation::Create<FStatus>(), InFilenames);
+	}
+
+	bool GetStateSucceeded = SourceControlProvider.GetState(InFilenames, FileStates, EStateCacheUsage::Use) == ECommandResult::Succeeded;
+
+	if (GetStateSucceeded && (!FileStates.IsEmpty()))
+	{
+		for (FSourceControlStateRef FileState : FileStates)
+		{
+			bool bIsCheckoutCompliant = (!bCheckCheckout) || (!FileState->IsCheckedOut());
+			bool bIsStatusCompliant = (!bCheckStatus) || FileState->IsModified();
+
+			if (bIsCheckoutCompliant && bIsStatusCompliant)
+			{
+				Files.Add(FileState);
+			}
+		}
+	}
+}
+
+void FUncontrolledChangelistState::UpdateStatus()
+{
+	TArray<FString> FilesToUpdate;
+	Algo::Transform(Files, FilesToUpdate, [](const FSourceControlStateRef& State) { return State->GetFilename(); });
+
+	if (FilesToUpdate.Num() == 0)
+	{
+		return;
+	}
+
+	ISourceControlProvider& SourceControlProvider = ISourceControlModule::Get().GetProvider();
+	SourceControlProvider.Execute(ISourceControlOperation::Create<FStatus>(), FilesToUpdate, EConcurrency::Asynchronous, FSourceControlOperationComplete::CreateSP(this, &FUncontrolledChangelistState::OnUpdateStatusCompleted));
+}
+
+void FUncontrolledChangelistState::OnUpdateStatusCompleted(const FSourceControlOperationRef& InOperation, ECommandResult::Type InResult)
+{
+	for (auto It = Files.CreateIterator(); It; ++It)
+	{
+		FSourceControlStateRef& State = *It;
+
+		if (State->IsCheckedOut() || (!State->IsModified()))
+		{
+			It.RemoveCurrent();
+		}
+	}
+}
+
+#undef LOCTEXT_NAMESPACE
diff --git a/Engine/Source/Developer/UncontrolledChangelists/Private/UncontrolledChangelistsModule.cpp b/Engine/Source/Developer/UncontrolledChangelists/Private/UncontrolledChangelistsModule.cpp
new file mode 100644
--- /dev/null
+++ b/Engine/Source/Developer/UncontrolledChangelists/Private/UncontrolledChangelistsModule.cpp
@@ -0,0 +1,267 @@
+// Copyright Epic Games, Inc. All Rights Reserved.
+
+#include "UncontrolledChangelistsModule.h"
+
+#include "Algo/Find.h"
+#include "Algo/Transform.h"
+#include "AssetRegistryModule.h"
+#include "Dom/JsonObject.h"
+#include "HAL/IConsoleManager.h"
+#include "ISourceControlModule.h"
+#include "Misc/FileHelper.h"
+#include "Misc/Paths.h"
+#include "Serialization/JsonSerializer.h"
+#include "Serialization/JsonWriter.h"
+
+#define LOCTEXT_NAMESPACE "UncontrolledChangelists"
+
+static TAutoConsoleVariable<bool> CVarUncontrolledChangelistsEnable(
+	TEXT("UncontrolledChangelists.Enable"),
+	false,
+	TEXT("Enables Uncontrolled Changelists (experimental).")
+);
+
+void FUncontrolledChangelistsModule::StartupModule()
+{
+	// Adds Default Uncontrolled Changelist if it is not already present.
+	FUncontrolledChangelist DefaultUncontrolledChangelist(FUncontrolledChangelist::DEFAULT_UNCONTROLLED_CHANGELIST_GUID, FUncontrolledChangelist::DEFAULT_UNCONTROLLED_CHANGELIST_NAME.ToString());
+	UncontrolledChangelistsStateCache.FindOrAdd(MoveTemp(DefaultUncontrolledChangelist), MakeShareable(new FUncontrolledChangelistState(DefaultUncontrolledChangelist)));
+
+	LoadState();
+
+	FAssetRegistryModule& AssetRegistryModule = FModuleManager::GetModuleChecked<FAssetRegistryModule>(TEXT("AssetRegistry"));
+	OnAssetAddedDelegateHandle = AssetRegistryModule.Get().OnAssetAdded().AddLambda([](const struct FAssetData& AssetData) { Get().OnAssetAdded(AssetData); });
+	OnAssetLoadedDelegateHandle = FCoreUObjectDelegates::OnAssetLoaded.AddLambda([](UObject* InAsset) { Get().OnAssetLoaded(InAsset); });
+	OnObjectTransactedDelegateHandle = FCoreUObjectDelegates::OnObjectTransacted.AddLambda([](UObject* InObject, const class FTransactionObjectEvent& InTransactionEvent)
+	{
+		Get().OnObjectTransacted(InObject, InTransactionEvent);
+	});
+}
+
+void FUncontrolledChangelistsModule::ShutdownModule()
+{
+	FAssetRegistryModule* AssetRegistryModulePtr = static_cast<FAssetRegistryModule*>(FModuleManager::Get().GetModule(TEXT("AssetRegistry")));
+
+	// Check in case AssetRegistry has already been shutdown.
+	if (AssetRegistryModulePtr != nullptr)
+	{
+		AssetRegistryModulePtr->Get().OnAssetAdded().Remove(OnAssetAddedDelegateHandle);
+	}
+
+	FCoreUObjectDelegates::OnAssetLoaded.Remove(OnAssetLoadedDelegateHandle);
+	FCoreUObjectDelegates::OnObjectTransacted.Remove(OnObjectTransactedDelegateHandle);
+}
+
+bool FUncontrolledChangelistsModule::IsEnabled() const
+{
+	return CVarUncontrolledChangelistsEnable.GetValueOnGameThread();
+}
+
+TArray<FUncontrolledChangelistStateRef> FUncontrolledChangelistsModule::GetChangelistStates() const
+{
+	TArray<FUncontrolledChangelistStateRef> UncontrolledChangelistStates;
+
+	if (IsEnabled())
+	{
+		Algo::Transform(UncontrolledChangelistsStateCache, UncontrolledChangelistStates, [](const auto& Pair) { return Pair.Value; });
+	}
+
+	return UncontrolledChangelistStates;
+}
+
+void FUncontrolledChangelistsModule::OnMakeWritable(const FString& InFilename)
+{
+	if (!IsEnabled())
+	{
+		return;
+	}
+
+	FUncontrolledChangelist DefaultUncontrolledChangelist(FUncontrolledChangelist::DEFAULT_UNCONTROLLED_CHANGELIST_GUID, FUncontrolledChangelist::DEFAULT_UNCONTROLLED_CHANGELIST_NAME.ToString());
+	FUncontrolledChangelistsStateCache::ValueType& UncontrolledChangelistState = UncontrolledChangelistsStateCache.FindOrAdd(MoveTemp(DefaultUncontrolledChangelist), MakeShareable(new FUncontrolledChangelistState(DefaultUncontrolledChangelist)));
+
+	UncontrolledChangelistState->AddFiles({ InFilename }, FUncontrolledChangelistState::ECheckFlags::NotCheckedOut);
+
+	SaveState();
+}
+
+void FUncontrolledChangelistsModule::UpdateStatus()
+{
+	if (!IsEnabled())
+	{
+		return;
+	}
+
+	for (FUncontrolledChangelistsStateCache::ElementType& Pair : UncontrolledChangelistsStateCache)
+	{
+		FUncontrolledChangelistsStateCache::ValueType& UncontrolledChangelistState = Pair.Value;
+		UncontrolledChangelistState->UpdateStatus();
+	}
+}
+
+void FUncontrolledChangelistsModule::OnAssetAdded(const struct FAssetData& AssetData)
+{
+
+}
+
+void FUncontrolledChangelistsModule::OnAssetLoaded(UObject* InAsset)
+{
+	FString Fullpath = GetUObjectPackageFullpath(InAsset);
+
+	if (Fullpath.IsEmpty())
+	{
+		return;
+	}
+}
+
+void FUncontrolledChangelistsModule::OnObjectTransacted(UObject* InObject, const class FTransactionObjectEvent& InTransactionEvent)
+{
+	FString Fullpath = GetUObjectPackageFullpath(InObject);
+
+	if (Fullpath.IsEmpty())
+	{
+		return;
+	}
+}
+
+void FUncontrolledChangelistsModule::OnFilesMovedToUncontrolledChangelist(const TArray<FString> InFilenames, const FUncontrolledChangelist& InChangelist)
+{
+	if (!IsEnabled())
+	{
+		return;
+	}
+
+	FUncontrolledChangelistsStateCache::ValueType* ChangelistState = UncontrolledChangelistsStateCache.Find(InChangelist);
+
+	if (ChangelistState == nullptr)
+	{
+		return;
+	}
+
+	(*ChangelistState)->AddFiles(InFilenames, FUncontrolledChangelistState::ECheckFlags::None);
+
+	SaveState();
+}
+
+void FUncontrolledChangelistsModule::SaveState() const
+{
+	TSharedPtr<FJsonObject> RootObject = MakeShareable(new FJsonObject);
+	TArray<TSharedPtr<FJsonValue>> UncontrolledChangelistsArray;
+
+	RootObject->SetNumberField(VERSION_NAME, VERSION_NUMBER);
+
+	for (const auto& Pair : UncontrolledChangelistsStateCache)
+	{
+		const FUncontrolledChangelist& UncontrolledChangelist = Pair.Key;
+		FUncontrolledChangelistStateRef UncontrolledChangelistState = Pair.Value;
+		TSharedPtr<FJsonObject> UncontrolledChangelistObject = MakeShareable(new FJsonObject);
+
+		UncontrolledChangelist.Serialize(UncontrolledChangelistObject.ToSharedRef());
+		UncontrolledChangelistState->Serialize(UncontrolledChangelistObject.ToSharedRef());
+
+		UncontrolledChangelistsArray.Add(MakeShareable(new FJsonValueObject(UncontrolledChangelistObject)));
+	}
+
+	RootObject->SetArrayField(CHANGELISTS_NAME, UncontrolledChangelistsArray);
+
+	using FStringWriter = TJsonWriter<TCHAR, TPrettyJsonPrintPolicy<TCHAR>>;
+	using FStringWriterFactory = TJsonWriterFactory<TCHAR, TPrettyJsonPrintPolicy<TCHAR>>;
+
+	FString RootObjectStr;
+	TSharedRef<FStringWriter> Writer = FStringWriterFactory::Create(&RootObjectStr);
+	FJsonSerializer::Serialize(RootObject.ToSharedRef(), Writer);
+
+	FFileHelper::SaveStringToFile(RootObjectStr, *GetPersistentFilePath());
+}
+
+void FUncontrolledChangelistsModule::LoadState()
+{
+	FString ImportJsonString;
+	TSharedPtr<FJsonObject> RootObject;
+	uint32 VersionNumber = 0;
+	const TArray<TSharedPtr<FJsonValue>>* UncontrolledChangelistsArray = nullptr;
+
+	if (!FFileHelper::LoadFileToString(ImportJsonString, *GetPersistentFilePath()))
+	{
+		UE_LOG(LogSourceControl, Warning, TEXT("Cannot load file %s."), *GetPersistentFilePath());
+		return;
+	}
+
+	TSharedRef<TJsonReader<>> JsonReader = TJsonReaderFactory<>::Create(ImportJsonString);
+
+	if (!FJsonSerializer::Deserialize(JsonReader, RootObject))
+	{
+		UE_LOG(LogSourceControl, Error, TEXT("Cannot deserialize RootObject."));
+		return;
+	}
+
+	if (!RootObject->TryGetNumberField(VERSION_NAME, VersionNumber))
+	{
+		UE_LOG(LogSourceControl, Error, TEXT("Cannot get field %s."), VERSION_NAME);
+		return;
+	}
+
+	if (VersionNumber != VERSION_NUMBER)
+	{
+		UE_LOG(LogSourceControl, Error, TEXT("Version number is invalid (file: %u, current: %u)."), VersionNumber, VERSION_NUMBER);
+		return;
+	}
+
+	if (!RootObject->TryGetArrayField(CHANGELISTS_NAME, UncontrolledChangelistsArray))
+	{
+		UE_LOG(LogSourceControl, Error, TEXT("Cannot get field %s."), CHANGELISTS_NAME);
+		return;
+	}
+
+	for (const TSharedPtr<FJsonValue>& JsonValue : *UncontrolledChangelistsArray)
+	{
+		FUncontrolledChangelist TempKey;
+		TSharedRef<FJsonObject> JsonObject = JsonValue->AsObject().ToSharedRef();
+
+		if (!TempKey.Deserialize(JsonObject))
+		{
+			UE_LOG(LogSourceControl, Error, TEXT("Cannot deserialize FUncontrolledChangelist."));
+			continue;
+		}
+
+		FUncontrolledChangelistsStateCache::ValueType& UncontrolledChangelistState = UncontrolledChangelistsStateCache.FindOrAdd(MoveTemp(TempKey), MakeShareable(new FUncontrolledChangelistState(TempKey)));
+
+		UncontrolledChangelistState->Deserialize(JsonObject);
+	}
+}
+
+FString FUncontrolledChangelistsModule::GetPersistentFilePath() const
+{
+	return FPaths::ProjectSavedDir() + TEXT("SourceControl/UncontrolledChangelists.json");
+}
+
+FString FUncontrolledChangelistsModule::GetUObjectPackageFullpath(const UObject* InObject) const
+{
+	FString Fullpath = TEXT("");
+
+	if (InObject == nullptr)
+	{
+		return Fullpath;
+	}
+
+	const UPackage* Package = InObject->GetPackage();
+
+	if (Package == nullptr)
+	{
+		return Fullpath;
+	}
+
+	const FString LocalFullPath(Package->GetLoadedPath().GetLocalFullPath());
+
+	if (LocalFullPath.IsEmpty())
+	{
+		return Fullpath;
+	}
+
+	Fullpath = FPaths::ConvertRelativePathToFull(LocalFullPath);
+
+	return Fullpath;
+}
+
+IMPLEMENT_MODULE(FUncontrolledChangelistsModule, UncontrolledChangelist);
+
+#undef LOCTEXT_NAMESPACE
diff --git a/Engine/Source/Developer/UncontrolledChangelists/Public/UncontrolledChangelist.h b/Engine/Source/Developer/UncontrolledChangelists/Public/UncontrolledChangelist.h
new file mode 100644
--- /dev/null
+++ b/Engine/Source/Developer/UncontrolledChangelists/Public/UncontrolledChangelist.h
@@ -0,0 +1,65 @@
+// Copyright Epic Games, Inc. All Rights Reserved.
+
+#pragma once
+
+#include "CoreMinimal.h"
+#include "ISourceControlChangelist.h"
+#include "Misc/Guid.h"
+
+class FUncontrolledChangelist : public TSharedFromThis<FUncontrolledChangelist, ESPMode::ThreadSafe>
+{
+public:
+	static constexpr TCHAR* GUID_NAME = TEXT("guid");
+	static constexpr TCHAR* NAME_NAME = TEXT("name");
+	static const FGuid DEFAULT_UNCONTROLLED_CHANGELIST_GUID;
+	static const FText DEFAULT_UNCONTROLLED_CHANGELIST_NAME;
+
+public:
+	FUncontrolledChangelist();
+	
+	FUncontrolledChangelist(const FGuid& InGuid, const FString& InName);
+
+	bool operator==(const FUncontrolledChangelist& InOther) const
+	{
+		return Guid == InOther.Guid;
+	}
+
+	bool operator!=(const FUncontrolledChangelist& InOther) const
+	{
+		return Guid != InOther.Guid;
+	}
+
+	bool IsDefault() const
+	{
+		return Guid == DEFAULT_UNCONTROLLED_CHANGELIST_GUID;
+	}
+
+	friend FORCEINLINE uint32 GetTypeHash(const FUncontrolledChangelist& PerforceUncontrolledChangelist)
+	{
+		return GetTypeHash(PerforceUncontrolledChangelist.Guid);
+	}
+
+	FString ToString() const
+	{
+		return Name.IsEmpty() ? Guid.ToString() : Name;
+	}
+
+	/**
+	 * Serialize the Uncontrolled Changelist to a Json Object.
+	 * @param	OutJsonObject	The Json Object used to serialize.
+	 */
+	void Serialize(TSharedRef<class FJsonObject> OutJsonObject) const;
+	
+	/**
+	 * Deserialize the Uncontrolled Changelist from a Json Object.
+	 * @param 	InJsonValue 	The Json Object to read from.
+	 * @return 	True if Deserialization succeeded.
+	 */
+	bool Deserialize(const TSharedRef<class FJsonObject> InJsonValue);
+
+private:
+	FGuid		Guid;
+	FString		Name;
+};
+
+typedef TSharedRef<FUncontrolledChangelist, ESPMode::ThreadSafe> FUncontrolledChangelistRef;
diff --git a/Engine/Source/Developer/UncontrolledChangelists/Public/UncontrolledChangelistState.h b/Engine/Source/Developer/UncontrolledChangelists/Public/UncontrolledChangelistState.h
new file mode 100644
--- /dev/null
+++ b/Engine/Source/Developer/UncontrolledChangelists/Public/UncontrolledChangelistState.h
@@ -0,0 +1,111 @@
+// Copyright Epic Games, Inc. All Rights Reserved.
+
+#pragma once
+
+#include "CoreMinimal.h"
+#include "ISourceControlChangelistState.h"
+#include "ISourceControlOperation.h"
+#include "ISourceControlProvider.h"
+
+#include "UncontrolledChangelist.h"
+
+class FUncontrolledChangelistState : public TSharedFromThis<FUncontrolledChangelistState, ESPMode::ThreadSafe>
+{
+public:
+	static constexpr TCHAR* FILES_NAME = TEXT("files");
+
+	enum class ECheckFlags
+	{
+		None			= 0,
+		Modified		= 1,
+		NotCheckedOut	= 1 << 1
+	};
+
+public:
+	FUncontrolledChangelistState(const FUncontrolledChangelist& InUncontrolledChangelist);
+
+	/**
+	 * Get the name of the icon graphic we should use to display the state in a UI.
+	 * @returns the name of the icon to display
+	 */
+	UNCONTROLLEDCHANGELISTS_API FName GetIconName() const;
+
+	/**
+	 * Get the name of the small icon graphic we should use to display the state in a UI.
+	 * @returns the name of the icon to display
+	 */
+	UNCONTROLLEDCHANGELISTS_API FName GetSmallIconName() const;
+
+	/**
+	 * Get a text representation of the state
+	 * @returns	the text to display for this state
+	 */
+	UNCONTROLLEDCHANGELISTS_API FText GetDisplayText() const;
+
+	/**
+	 * Get a text representation of the state
+	 * @returns	the text to display for this state
+	 */
+	UNCONTROLLEDCHANGELISTS_API FText GetDescriptionText() const;
+
+	/**
+	 * Get a tooltip to describe this state
+	 * @returns	the text to display for this states tooltip
+	 */
+	UNCONTROLLEDCHANGELISTS_API FText GetDisplayTooltip() const;
+
+	/**
+	 * Get the timestamp of the last update that was made to this state.
+	 * @returns	the timestamp of the last update
+	 */
+	const FDateTime& GetTimeStamp() const;
+
+	UNCONTROLLEDCHANGELISTS_API const TSet<FSourceControlStateRef>& GetFilesStates() const;
+
+	/**
+	 * Serialize the state of the Uncontrolled Changelist to a Json Object.
+	 * @param 	OutJsonObject 	The Json object used to serialize.
+	 */
+	void Serialize(TSharedRef<class FJsonObject> OutJsonObject) const;
+
+	/**
+	 * Deserialize the state of the Uncontrolled Changelist from a Json Object.
+	 * @param 	InJsonValue 	The Json Object to read from.
+	 * @return 	True if Deserialization succeeded.
+	 */
+	bool Deserialize(const TSharedRef<FJsonObject> InJsonValue);
+
+	/**
+	 * Adds a file to this Uncontrolled Changelist State.
+	 * @param 	InFilename 		The file to be added.
+	 * @param 	bInCheckStatus 	Tells whether we should verify if the file has been modified.
+	 */	
+	void AddFiles(const TArray<FString>& InFilenames, const ECheckFlags InCheckFlags);
+
+	/**
+	 * Updates the status of all files.
+	 */
+	void UpdateStatus();
+
+private:
+
+	/**
+	 * Called upon completion of FStatus operation
+	 * @param 	InOperation 	The operation just completed.
+	 * @param 	InResult 		The result of the operation.
+	 */
+	void OnUpdateStatusCompleted(const FSourceControlOperationRef& InOperation, ECommandResult::Type InResult);
+
+public:
+	FUncontrolledChangelist Changelist;
+
+	FString Description;
+
+	TSet<FSourceControlStateRef> Files;
+
+	/** The timestamp of the last update */
+	FDateTime TimeStamp;
+};
+
+ENUM_CLASS_FLAGS(FUncontrolledChangelistState::ECheckFlags);
+typedef TSharedRef<FUncontrolledChangelistState, ESPMode::ThreadSafe> FUncontrolledChangelistStateRef;
diff --git a/Engine/Source/Developer/UncontrolledChangelists/Public/UncontrolledChangelistsModule.h b/Engine/Source/Developer/UncontrolledChangelists/Public/UncontrolledChangelistsModule.h
new file mode 100644
--- /dev/null
+++ b/Engine/Source/Developer/UncontrolledChangelists/Public/UncontrolledChangelistsModule.h
@@ -0,0 +1,118 @@
+// Copyright Epic Games, Inc. All Rights Reserved.
+
+#pragma once
+
+#include "CoreMinimal.h"
+
+#include "UncontrolledChangelistState.h"
+#include "Modules/ModuleInterface.h"
+#include "Modules/ModuleManager.h"
+#include "UObject/UObjectGlobals.h"
+
+/**
+ * Interface for talking to Uncontrolled Changelists
+ */
+class UNCONTROLLEDCHANGELISTS_API FUncontrolledChangelistsModule : public IModuleInterface
+{
+	typedef TMap<FUncontrolledChangelist, TSharedRef<class FUncontrolledChangelistState, ESPMode::ThreadSafe>> FUncontrolledChangelistsStateCache;
+
+public:
+	static constexpr TCHAR* VERSION_NAME = TEXT("version");
+	static constexpr TCHAR* CHANGELISTS_NAME = TEXT("changelists");
+	static constexpr uint32 VERSION_NUMBER = 0;
+
+public:
+	/** IModuleInterface implementation */
+	virtual void StartupModule() override;
+	virtual void ShutdownModule() override;
+
+	/**
+	 * Check whether uncontrolled changelist module is enabled.
+	 */
+	bool IsEnabled() const;
+
+	/**
+	 * Get the changelist state of each cached Uncontrolled Changelist.
+	 */
+	 TArray<FUncontrolledChangelistStateRef> GetChangelistStates() const;
+
+	/**
+	 * Called when a file has been made writable. Adds the file to the Default Uncontrolled Changelist
+	 * @param	InFilename			The file to be added.
+	 */
+	void OnMakeWritable(const FString& InFilename);
+
+
+	/**
+	 * Updates the status of Uncontrolled Changelists and files.
+	 */
+	void UpdateStatus();
+
+	/**
+	 * Gets a reference to the UncontrolledChangelists module
+	 * @return A reference to the UncontrolledChangelists module.
+	 */
+	static inline FUncontrolledChangelistsModule& Get()
+	{
+		static FName UncontrolledChangelistsModuleName("UncontrolledChangelists");
+		return FModuleManager::LoadModuleChecked<FUncontrolledChangelistsModule>(UncontrolledChangelistsModuleName);
+	}
+
+	/**
+	 * Delegate callback called when assets are added to AssetRegistry.
+	 * @param 	AssetData 	The asset just added.
+	 */
+	void OnAssetAdded(const struct FAssetData& AssetData);
+
+	/**
+	 * Delegate callback called when an asset is loaded.
+	 * @param 	InAsset 	The loaded asset.
+	 */
+	void OnAssetLoaded(UObject* InAsset);
+
+	/**
+	 * Delegate callback called when an object is transacted.
+	 * @param 	InObject 	        The transacted object.
+	 * @param 	InTransactionEvent 	The event representing the transaction.
+	 */
+	void OnObjectTransacted(UObject* InObject, const class FTransactionObjectEvent& InTransactionEvent);
+
+	/**
+	 * Called after an action moving files to an Uncontrolled Changelist (example: drag and drop).
+	 * Moves the files to the provided uncontrolled changelist.
+	 * @param 	InFilenames 	The files to move.
+	 * @param 	InChangelist 	The Uncontrolled Changelist where to move the files.
+	 * @return 	returnDesc
+	 */
+	void OnFilesMovedToUncontrolledChangelist(const TArray<FString> InFilenames, const FUncontrolledChangelist& InChangelist);
+
+private:
+	/**
+	 * Saves the state of UncontrolledChangelists to Json for persistency.
+	 */
+	void SaveState() const;
+	
+	/**
+	 * Restores the previously saved state from Json.
+	 */
+	void LoadState();
+
+	/**
+	 * Helper returning the location of the file used for persistency.
+	 * @return 	A string containing the filepath.
+	 */
+	FString GetPersistentFilePath() const;
+
+	/**
+	 * Helper returning the package path where an UObject is located.
+	 * @param 	InObject 	The object used to locate the package.
+	 * @return 	A String containing the filepath of the package.
+	 */
+	FString GetUObjectPackageFullpath(const UObject* InObject) const;
+
+private:
+	FUncontrolledChangelistsStateCache UncontrolledChangelistsStateCache;
+	FDelegateHandle OnAssetAddedDelegateHandle;
+	FDelegateHandle OnAssetLoadedDelegateHandle;
+	FDelegateHandle OnObjectTransactedDelegateHandle;
+};
diff --git a/Engine/Source/Developer/UncontrolledChangelists/UncontrolledChangelists.Build.cs b/Engine/Source/Developer/UncontrolledChangelists/UncontrolledChangelists.Build.cs
new file mode 100644
--- /dev/null
+++ b/Engine/Source/Developer/UncontrolledChangelists/UncontrolledChangelists.Build.cs
@@ -0,0 +1,44 @@
+// Copyright Epic Games, Inc. All Rights Reserved.
+
+using UnrealBuildTool;
+
+public class UncontrolledChangelists : ModuleRules
+{
+	public UncontrolledChangelists(ReadOnlyTargetRules Target) : base(Target)
+	{
+		PublicIncludePaths.Add(ModuleDirectory + "/Public");
+		PrivateIncludePaths.Add(ModuleDirectory + "/Private");
+
+		PublicDependencyModuleNames.AddRange(
+			new string[]
+			{
+				"Core",
+				"CoreUObject",
+				"SlateCore",
+				"InputCore",
+			}
+		);
+
+		PrivateDependencyModuleNames.AddRange(
+			new string[]
+			{
+                "EditorStyle",
+				"Json",
+				"SourceControl",
+				"AssetRegistry",
+			}
+		);
+
+		if (Target.bBuildEditor)
+        {
+			PrivateDependencyModuleNames.AddRange(
+                new string[]
+				{
+					"EditorFramework",
+					"Engine",
+					"UnrealEd",
+				}
+			);
+        }
+	}
+}
diff --git a/Engine/Source/Editor/EditorStyle/Private/StarshipStyle.cpp b/Engine/Source/Editor/EditorStyle/Private/StarshipStyle.cpp
--- a/Engine/Source/Editor/EditorStyle/Private/StarshipStyle.cpp
+++ b/Engine/Source/Editor/EditorStyle/Private/StarshipStyle.cpp
@@ -7213,6 +7213,8 @@ void FStarshipEditorStyle::FStyle::SetupSourceControlStyles()
 		Set("SourceControl.ChangelistsTab", new CORE_IMAGE_BRUSH_SVG("Starship/Common/check-circle", Icon16x16));
 		Set("SourceControl.Changelist", new CORE_IMAGE_BRUSH_SVG("Starship/SourceControl/SCC_Changelist", Icon16x16, FStyleColors::AccentRed));
 		Set("SourceControl.ShelvedChangelist", new CORE_IMAGE_BRUSH_SVG("Starship/SourceControl/SCC_Changelist", Icon16x16, FStyleColors::AccentBlue));
+		Set("SourceControl.UncontrolledChangelist", new CORE_IMAGE_BRUSH_SVG("Starship/SourceControl/SCC_Changelist", Icon32x32, FStyleColors::AccentGray));
+		Set("SourceControl.UncontrolledChangelist_Small", new CORE_IMAGE_BRUSH_SVG("Starship/SourceControl/SCC_Changelist", Icon16x16, FStyleColors::AccentGray));
 
 		Set("SourceControl.Add", new CORE_IMAGE_BRUSH_SVG( "Starship/SourceControl/SCC_ContentAdd",Icon16x16));
 		Set("SourceControl.Edit", new CORE_IMAGE_BRUSH_SVG( "Starship/SourceControl/SCC_CheckedOut", Icon16x16));
diff --git a/Engine/Source/Editor/SourceControlWindows/Private/SSourceControlChangelists.cpp b/Engine/Source/Editor/SourceControlWindows/Private/SSourceControlChangelists.cpp
--- a/Engine/Source/Editor/SourceControlWindows/Private/SSourceControlChangelists.cpp
+++ b/Engine/Source/Editor/SourceControlWindows/Private/SSourceControlChangelists.cpp
@@ -18,6 +18,7 @@
 
 #include "ISourceControlProvider.h"
 #include "ISourceControlModule.h"
+#include "UncontrolledChangelistsModule.h"
 #include "SourceControlOperations.h"
 #include "ToolMenus.h"
 #include "Widgets/Images/SLayeredImage.h"
@@ -195,6 +196,9 @@ void SSourceControlChangelistsWidget::RequestRefresh()
 
 		ISourceControlProvider& SourceControlProvider = ISourceControlModule::Get().GetProvider();
 		SourceControlProvider.Execute(UpdatePendingChangelistsOperation, EConcurrency::Asynchronous);
+
+		FUncontrolledChangelistsModule& UncontrolledChangelistModule = FUncontrolledChangelistsModule::Get();
+		UncontrolledChangelistModule.UpdateStatus();
 	}
 	else
 	{
@@ -238,7 +242,9 @@ void SSourceControlChangelistsWidget::Refresh()
 		SaveExpandedState(ExpandedStates);
 
 		ISourceControlProvider& SourceControlProvider = ISourceControlModule::Get().GetProvider();
+		FUncontrolledChangelistsModule& UncontrolledChangelistModule = FUncontrolledChangelistsModule::Get();
 		TArray<FSourceControlChangelistRef> Changelists = SourceControlProvider.GetChangelists(EStateCacheUsage::Use);
+		TArray<FUncontrolledChangelistStateRef> UncontrolledChangelistStates = UncontrolledChangelistModule.GetChangelistStates();
 
 		TArray<FSourceControlChangelistStateRef> ChangelistsStates;
 		SourceControlProvider.GetState(Changelists, ChangelistsStates, EStateCacheUsage::Use);
@@ -247,12 +253,19 @@ void SSourceControlChangelistsWidget::Refresh()
 
 		// Count number of steps for slow task...
 		int32 ElementsToProcess = ChangelistsStates.Num();
+		ElementsToProcess += UncontrolledChangelistStates.Num();
+
 		for (FSourceControlChangelistStateRef ChangelistState : ChangelistsStates)
 		{
 			ElementsToProcess += ChangelistState->GetFilesStates().Num();
 			ElementsToProcess += ChangelistState->GetShelvedFilesStates().Num();
 		}
 
+		for (FUncontrolledChangelistStateRef UncontrolledChangelistState : UncontrolledChangelistStates)
+		{
+			ElementsToProcess += UncontrolledChangelistState->GetFilesStates().Num();
+		}
+
 		FScopedSlowTask SlowTask(ElementsToProcess, LOCTEXT("SourceControl_RebuildTree", "Beautifying Changelist Tree Items Paths"));
 		SlowTask.MakeDialog(/*bShowCancelButton=*/true);
 
@@ -289,6 +302,23 @@ void SSourceControlChangelistsWidget::Refresh()
 			bBeautifyPaths &= !SlowTask.ShouldCancel();
 		}
 
+		for (FUncontrolledChangelistStateRef UncontrolledChangelistState : UncontrolledChangelistStates)
+		{
+			FChangelistTreeItemRef UncontrolledChangelistTreeItem = MakeShareable(new FUncontrolledChangelistTreeItem(UncontrolledChangelistState));
+			
+			for (FSourceControlStateRef FileRef : UncontrolledChangelistState->GetFilesStates())
+			{
+				FChangelistTreeItemRef FileTreeItem = MakeShareable(new FFileTreeItem(FileRef, bBeautifyPaths));
+				UncontrolledChangelistTreeItem->AddChild(FileTreeItem);
+				SlowTask.EnterProgressFrame();
+				bBeautifyPaths &= !SlowTask.ShouldCancel();
+			}
+
+			ChangelistsNodes.Add(UncontrolledChangelistTreeItem);
+			SlowTask.EnterProgressFrame();
+			bBeautifyPaths &= !SlowTask.ShouldCancel();
+		}
+
 		RestoreExpandedState(ExpandedStates);
 
 		TreeView->RequestTreeRefresh();
@@ -1105,7 +1135,8 @@ public:
 	{
 		if (ColumnName == TEXT("Change"))
 		{
-			const FSlateBrush* IconBrush = FEditorStyle::GetBrush("SourceControl.Changelist");
+			const FSlateBrush* IconBrush = (TreeItem != nullptr) ? FEditorStyle::GetBrush(TreeItem->ChangelistState->GetSmallIconName())
+																 : FEditorStyle::GetBrush("SourceControl.Changelist");
 
 			return SNew(SHorizontalBox)
 				+ SHorizontalBox::Slot()
@@ -1191,6 +1222,125 @@ private:
 	FChangelistTreeItem* TreeItem;
 };
 
+class SUncontrolledChangelistTableRow : public SMultiColumnTableRow<FChangelistTreeItemPtr>
+{
+public:
+	SLATE_BEGIN_ARGS(SUncontrolledChangelistTableRow)
+		: _TreeItemToVisualize()
+	{
+	}
+	SLATE_ARGUMENT(FChangelistTreeItemPtr, TreeItemToVisualize)
+	SLATE_END_ARGS()
+
+public:
+	/**
+	* Construct child widgets that comprise this widget.
+	*
+	* @param InArgs Declaration from which to construct this widget.
+	*/
+	void Construct(const FArguments& InArgs, const TSharedRef<STableViewBase>& InOwner)
+	{
+		TreeItem = static_cast<FUncontrolledChangelistTreeItem*>(InArgs._TreeItemToVisualize.Get());
+
+		auto Args = FSuperRowType::FArguments();
+		SMultiColumnTableRow<FChangelistTreeItemPtr>::Construct(Args, InOwner);
+	}
+
+	// SMultiColumnTableRow overrides
+	virtual TSharedRef<SWidget> GenerateWidgetForColumn(const FName& ColumnName) override
+	{
+		if (ColumnName == TEXT("Change"))
+		{
+			const FSlateBrush* IconBrush = (TreeItem != nullptr) ? FEditorStyle::GetBrush(TreeItem->UncontrolledChangelistState->GetSmallIconName())
+				: FEditorStyle::GetBrush("SourceControl.Changelist");
+
+			return SNew(SHorizontalBox)
+				+ SHorizontalBox::Slot()
+				.AutoWidth()
+				.VAlign(VAlign_Center)
+				[
+					SNew(SExpanderArrow, SharedThis(this))
+				]
+
+			+ SHorizontalBox::Slot()
+				.AutoWidth()
+				.VAlign(VAlign_Center)
+				[
+					SNew(SImage)
+					.Image(IconBrush)
+				]
+
+			+ SHorizontalBox::Slot()
+				.AutoWidth()
+				.Padding(2.0f, 0.0f)
+				.VAlign(VAlign_Center)
+				[
+					SNew(STextBlock)
+					.Text(this, &SUncontrolledChangelistTableRow::GetChangelistText)
+				];
+		}
+		else if (ColumnName == TEXT("Description"))
+		{
+			return SNew(SHorizontalBox)
+				+ SHorizontalBox::Slot()
+				.AutoWidth()
+				.Padding(2.0f, 0.0f)
+				.VAlign(VAlign_Center)
+				[
+					SNew(STextBlock)
+					.Text(this, &SUncontrolledChangelistTableRow::GetChangelistDescriptionText)
+				];
+		}
+		else
+		{
+			return SNullWidget::NullWidget;
+		}
+	}
+
+	FText GetChangelistText() const
+	{
+		return TreeItem->GetDisplayText();
+	}
+
+	FText GetChangelistDescriptionText() const
+	{
+		FString DescriptionString = TreeItem->GetDescriptionText().ToString();
+		// Here we'll both remove \r\n (when edited from the dialog) and \n (when we get it from the SCC)
+		DescriptionString.ReplaceInline(TEXT("\r"), TEXT(""));
+		DescriptionString.ReplaceInline(TEXT("\n"), TEXT(" "));
+		DescriptionString.TrimEndInline();
+		return FText::FromString(DescriptionString);
+	}
+
+protected:
+	//~ Begin STableRow Interface.
+	virtual FReply OnDrop(const FGeometry& InGeometry, const FDragDropEvent& InDragDropEvent) override
+	{
+		TSharedPtr<FSCCFileDragDropOp> Operation = InDragDropEvent.GetOperationAs<FSCCFileDragDropOp>();
+		
+		if (Operation.IsValid())
+		{
+			TArray<FString> Files;
+			Algo::Transform(Operation->Files, Files, [](const FSourceControlStateRef& State) { return State->GetFilename(); });
+
+			ISourceControlProvider& SourceControlProvider = ISourceControlModule::Get().GetProvider();
+			auto RevertOperation = ISourceControlOperation::Create<FRevert>();
+
+			RevertOperation->SetSoftRevert(true);
+			SourceControlProvider.Execute(RevertOperation, Files);
+
+			FUncontrolledChangelistsModule::Get().OnFilesMovedToUncontrolledChangelist(Files, TreeItem->UncontrolledChangelistState->Changelist);
+		}
+
+		return FReply::Handled();
+	}
+	//~ End STableRow Interface.
+
+private:
+	/** The info about the widget that we are visualizing. */
+	FUncontrolledChangelistTreeItem* TreeItem;
+};
+
 class SFileTableRow : public SMultiColumnTableRow<FChangelistTreeItemPtr>
 {
 public:
@@ -1382,6 +1532,10 @@ TSharedRef<ITableRow> SSourceControlChangelistsWidget::OnGenerateRow(FChangelist
 		return SNew(SChangelistTableRow, OwnerTable)
 			.TreeItemToVisualize(InTreeItem);
 
+	case IChangelistTreeItem::UncontrolledChangelist:
+		return SNew(SUncontrolledChangelistTableRow, OwnerTable)
+			.TreeItemToVisualize(InTreeItem);
+
 	case IChangelistTreeItem::File:
 		return SNew(SFileTableRow, OwnerTable)
 			.TreeItemToVisualize(InTreeItem)
diff --git a/Engine/Source/Editor/SourceControlWindows/Private/SSourceControlCommon.h b/Engine/Source/Editor/SourceControlWindows/Private/SSourceControlCommon.h
--- a/Engine/Source/Editor/SourceControlWindows/Private/SSourceControlCommon.h
+++ b/Engine/Source/Editor/SourceControlWindows/Private/SSourceControlCommon.h
@@ -6,6 +6,7 @@
 #include "Styling/SlateColor.h"
 #include "Styling/SlateTypes.h"
 #include "ISourceControlProvider.h"
+#include "UncontrolledChangelistState.h"
 
 struct FAssetData;
 struct IChangelistTreeItem;
@@ -18,6 +19,7 @@ struct IChangelistTreeItem : TSharedFromThis<IChangelistTreeItem>
 	{
 		Invalid,
 		Changelist,
+		UncontrolledChangelist,
 		File,
 		ShelvedChangelist, // container for shelved files
 		ShelvedFile
@@ -72,6 +74,27 @@ struct FChangelistTreeItem : public IChangelistTreeItem
 	FSourceControlChangelistStateRef ChangelistState;
 };
 
+struct FUncontrolledChangelistTreeItem : public IChangelistTreeItem
+{
+	FUncontrolledChangelistTreeItem(FUncontrolledChangelistStateRef InUncontrolledChangelistState)
+		: UncontrolledChangelistState(InUncontrolledChangelistState)
+	{
+		Type = IChangelistTreeItem::UncontrolledChangelist;
+	}
+
+	FText GetDisplayText() const
+	{
+		return UncontrolledChangelistState->GetDisplayText();
+	}
+
+	FText GetDescriptionText() const
+	{
+		return UncontrolledChangelistState->GetDescriptionText();
+	}
+
+	FUncontrolledChangelistStateRef UncontrolledChangelistState;
+};
+
 struct FShelvedChangelistTreeItem : public IChangelistTreeItem
 {
 	FShelvedChangelistTreeItem()
diff --git a/Engine/Source/Editor/SourceControlWindows/SourceControlWindows.Build.cs b/Engine/Source/Editor/SourceControlWindows/SourceControlWindows.Build.cs
--- a/Engine/Source/Editor/SourceControlWindows/SourceControlWindows.Build.cs
+++ b/Engine/Source/Editor/SourceControlWindows/SourceControlWindows.Build.cs
@@ -20,6 +20,7 @@ public class SourceControlWindows : ModuleRules
 				"SlateCore",
                 "EditorStyle",
 				"SourceControl", 
+				"UncontrolledChangelists",
 				"AssetTools",
 				"EditorFramework",
 				"WorkspaceMenuStructure",
diff --git a/Engine/Source/Editor/UnrealEd/Private/FileHelpers.cpp b/Engine/Source/Editor/UnrealEd/Private/FileHelpers.cpp
--- a/Engine/Source/Editor/UnrealEd/Private/FileHelpers.cpp
+++ b/Engine/Source/Editor/UnrealEd/Private/FileHelpers.cpp
@@ -24,6 +24,7 @@
 #include "Engine/MapBuildDataRegistry.h"
 #include "Editor/EditorEngine.h"
 #include "ISourceControlModule.h"
+#include "UncontrolledChangelistsModule.h"
 #include "SourceControlOperations.h"
 #include "Editor/UnrealEdEngine.h"
 #include "Settings/EditorLoadingSavingSettings.h"
@@ -1617,6 +1618,7 @@ bool FEditorFileUtils::PromptToCheckoutPackages(bool bCheckDirty, const TArray<U
 	bool bResult = true;
 
 	ISourceControlProvider& SourceControlProvider = ISourceControlModule::Get().GetProvider();
+	FUncontrolledChangelistsModule& UncontrolledChangelistModule = FUncontrolledChangelistsModule::Get();
 	
 	// The checkout dialog to show users if any packages need to be checked out
 	const FText DialogTitle = NSLOCTEXT("PackagesDialogModule", "CheckoutPackagesDialogTitle", "Check Out Assets");
@@ -1730,6 +1732,8 @@ bool FEditorFileUtils::PromptToCheckoutPackages(bool bCheckDirty, const TArray<U
 						// Knock off the read only flag from the current file attributes
 						if (FPlatformFileManager::Get().GetPlatformFile().SetReadOnly(*Filename, false))
 						{
+							UncontrolledChangelistModule.OnMakeWritable(Filename);
+
 							PackagesNotToPromptAnyMore.Add(PackageToMakeWritable->GetName());
 							if (OutPackagesCheckedOutOrMadeWritable)
 							{
diff --git a/Engine/Source/Editor/UnrealEd/UnrealEd.Build.cs b/Engine/Source/Editor/UnrealEd/UnrealEd.Build.cs
--- a/Engine/Source/Editor/UnrealEd/UnrealEd.Build.cs
+++ b/Engine/Source/Editor/UnrealEd/UnrealEd.Build.cs
@@ -88,6 +88,7 @@ public class UnrealEd : ModuleRules
 				"EditorFramework",
 				"EditorStyle",
 				"SourceControl",
+				"UncontrolledChangelists",
 				"UnrealEdMessages",
 				"GameplayDebugger",
 				"BlueprintGraph",
@@ -327,6 +328,7 @@ public class UnrealEd : ModuleRules
 				"AIModule",
 				"Engine",
 				"SourceControl",
+				"UncontrolledChangelists",
 			}
 		);
 
