Fixes compilation error for Linux on Uncontrolled Changelist Module.

diff --git a/Engine/Source/Developer/UncontrolledChangelists/Public/UncontrolledChangelist.h b/Engine/Source/Developer/UncontrolledChangelists/Public/UncontrolledChangelist.h
--- a/Engine/Source/Developer/UncontrolledChangelists/Public/UncontrolledChangelist.h
+++ b/Engine/Source/Developer/UncontrolledChangelists/Public/UncontrolledChangelist.h
@@ -9,8 +9,8 @@
 class FUncontrolledChangelist : public TSharedFromThis<FUncontrolledChangelist, ESPMode::ThreadSafe>
 {
 public:
-	static constexpr TCHAR* GUID_NAME = TEXT("guid");
-	static constexpr TCHAR* NAME_NAME = TEXT("name");
+	static constexpr const TCHAR* GUID_NAME = TEXT("guid");
+	static constexpr const TCHAR* NAME_NAME = TEXT("name");
 	static const FGuid DEFAULT_UNCONTROLLED_CHANGELIST_GUID;
 	static const FText DEFAULT_UNCONTROLLED_CHANGELIST_NAME;
 
diff --git a/Engine/Source/Developer/UncontrolledChangelists/Public/UncontrolledChangelistState.h b/Engine/Source/Developer/UncontrolledChangelists/Public/UncontrolledChangelistState.h
--- a/Engine/Source/Developer/UncontrolledChangelists/Public/UncontrolledChangelistState.h
+++ b/Engine/Source/Developer/UncontrolledChangelists/Public/UncontrolledChangelistState.h
@@ -12,7 +12,7 @@
 class FUncontrolledChangelistState : public TSharedFromThis<FUncontrolledChangelistState, ESPMode::ThreadSafe>
 {
 public:
-	static constexpr TCHAR* FILES_NAME = TEXT("files");
+	static constexpr const TCHAR* FILES_NAME = TEXT("files");
 
 	enum class ECheckFlags
 	{
diff --git a/Engine/Source/Developer/UncontrolledChangelists/Public/UncontrolledChangelistsModule.h b/Engine/Source/Developer/UncontrolledChangelists/Public/UncontrolledChangelistsModule.h
--- a/Engine/Source/Developer/UncontrolledChangelists/Public/UncontrolledChangelistsModule.h
+++ b/Engine/Source/Developer/UncontrolledChangelists/Public/UncontrolledChangelistsModule.h
@@ -17,8 +17,8 @@ class UNCONTROLLEDCHANGELISTS_API FUncontrolledChangelistsModule : public IModul
 	typedef TMap<FUncontrolledChangelist, TSharedRef<class FUncontrolledChangelistState, ESPMode::ThreadSafe>> FUncontrolledChangelistsStateCache;
 
 public:
-	static constexpr TCHAR* VERSION_NAME = TEXT("version");
-	static constexpr TCHAR* CHANGELISTS_NAME = TEXT("changelists");
+	static constexpr const TCHAR* VERSION_NAME = TEXT("version");
+	static constexpr const TCHAR* CHANGELISTS_NAME = TEXT("changelists");
 	static constexpr uint32 VERSION_NUMBER = 0;
 
 public:
