Adds moving files from Uncontrolled Changelist to Changelist.

diff --git a/Engine/Plugins/Developer/PerforceSourceControl/Source/PerforceSourceControl/Private/PerforceSourceControlOperations.cpp b/Engine/Plugins/Developer/PerforceSourceControl/Source/PerforceSourceControl/Private/PerforceSourceControlOperations.cpp
--- a/Engine/Plugins/Developer/PerforceSourceControl/Source/PerforceSourceControl/Private/PerforceSourceControlOperations.cpp
+++ b/Engine/Plugins/Developer/PerforceSourceControl/Source/PerforceSourceControl/Private/PerforceSourceControlOperations.cpp
@@ -233,7 +233,7 @@ static bool CheckWorkspaceRecordSet(const FP4RecordSet& InRecords, TArray<FText>
 	return false;
 }
 
-static void AppendChangelistParameter(TArray<FString>& InOutParams)
+static bool AppendChangelistParameter(TArray<FString>& InOutParams)
 {
 	FPerforceSourceControlModule& PerforceSourceControl = FModuleManager::GetModuleChecked<FPerforceSourceControlModule>("PerforceSourceControl");
 	FPerforceSourceControlSettings& Settings = PerforceSourceControl.AccessSettings();
@@ -243,7 +243,11 @@ static void AppendChangelistParameter(TArray<FString>& InOutParams)
 	{
 		InOutParams.Add(TEXT("-c"));
 		InOutParams.Add(ChangelistNumber);
+
+		return true;
 	}
+
+	return false;
 }
 
 FName FPerforceConnectWorker::GetName() const
@@ -310,7 +314,16 @@ bool FPerforceCheckOutWorker::Execute(FPerforceSourceControlCommand& InCommand)
 		FPerforceConnection& Connection = ScopedConnection.GetConnection();
 		TArray<FString> Parameters;
 
-		AppendChangelistParameter(Parameters);
+		if ((!AppendChangelistParameter(Parameters)) && InCommand.Changelist.IsInitialized())
+		{
+			FString ChangelistNumber = InCommand.Changelist.ToString();
+
+			if (!ChangelistNumber.IsEmpty())
+			{
+				Parameters.Add(TEXT("-c"));
+				Parameters.Add(ChangelistNumber);
+			}
+		}
 
 		Parameters.Append(InCommand.Files);
 		FP4RecordSet Records;
diff --git a/Engine/Source/Developer/UncontrolledChangelists/Private/UncontrolledChangelistState.cpp b/Engine/Source/Developer/UncontrolledChangelists/Private/UncontrolledChangelistState.cpp
--- a/Engine/Source/Developer/UncontrolledChangelists/Private/UncontrolledChangelistState.cpp
+++ b/Engine/Source/Developer/UncontrolledChangelists/Private/UncontrolledChangelistState.cpp
@@ -3,6 +3,7 @@
 #include "UncontrolledChangelistState.h"
 
 #include "Algo/Transform.h"
+#include "Algo/ForEach.h"
 #include "Dom/JsonObject.h"
 #include "Dom/JsonValue.h"
 #include "ISourceControlModule.h"
@@ -121,6 +122,11 @@ void FUncontrolledChangelistState::AddFiles(const TArray<FString>& InFilenames,
 	}
 }
 
+void FUncontrolledChangelistState::RemoveFiles(const TArray<FSourceControlStateRef>& InFileStates)
+{
+	Algo::ForEach(InFileStates, [&Files = Files](const FSourceControlStateRef& FileState) { Files.Remove(FileState); });
+}
+
 void FUncontrolledChangelistState::UpdateStatus()
 {
 	TArray<FString> FilesToUpdate;
diff --git a/Engine/Source/Developer/UncontrolledChangelists/Private/UncontrolledChangelistsModule.cpp b/Engine/Source/Developer/UncontrolledChangelists/Private/UncontrolledChangelistsModule.cpp
--- a/Engine/Source/Developer/UncontrolledChangelists/Private/UncontrolledChangelistsModule.cpp
+++ b/Engine/Source/Developer/UncontrolledChangelists/Private/UncontrolledChangelistsModule.cpp
@@ -6,12 +6,17 @@
 #include "Algo/Transform.h"
 #include "AssetRegistryModule.h"
 #include "Dom/JsonObject.h"
+#include "FileHelpers.h"
 #include "HAL/IConsoleManager.h"
 #include "ISourceControlModule.h"
+#include "ISourceControlProvider.h"
 #include "Misc/FileHelper.h"
 #include "Misc/Paths.h"
+#include "PackagesDialog.h"
 #include "Serialization/JsonSerializer.h"
 #include "Serialization/JsonWriter.h"
+#include "SourceControlOperations.h"
+#include "Styling/SlateTypes.h"
 
 #define LOCTEXT_NAMESPACE "UncontrolledChangelists"
 
@@ -123,25 +128,115 @@ void FUncontrolledChangelistsModule::OnObjectTransacted(UObject* InObject, const
 	}
 }
 
-void FUncontrolledChangelistsModule::OnFilesMovedToUncontrolledChangelist(const TArray<FString> InFilenames, const FUncontrolledChangelist& InChangelist)
+void FUncontrolledChangelistsModule::MoveFilesToUncontrolledChangelist(const TArray<FSourceControlStateRef>& InControlledFileStates, const TArray<FSourceControlStateRef>& InUncontrolledFileStates, const FUncontrolledChangelist& InUncontrolledChangelist)
 {
 	if (!IsEnabled())
 	{
 		return;
 	}
 
-	FUncontrolledChangelistsStateCache::ValueType* ChangelistState = UncontrolledChangelistsStateCache.Find(InChangelist);
+	FUncontrolledChangelistsStateCache::ValueType* ChangelistState = UncontrolledChangelistsStateCache.Find(InUncontrolledChangelist);
 
 	if (ChangelistState == nullptr)
 	{
 		return;
 	}
 
-	(*ChangelistState)->AddFiles(InFilenames, FUncontrolledChangelistState::ECheckFlags::None);
+	TArray<FString> Filenames;
+	Algo::Transform(InControlledFileStates, Filenames, [](const FSourceControlStateRef& State) { return State->GetFilename(); });
+
+	ISourceControlProvider& SourceControlProvider = ISourceControlModule::Get().GetProvider();
+	auto RevertOperation = ISourceControlOperation::Create<FRevert>();
+
+	// Revert controlled files
+	RevertOperation->SetSoftRevert(true);
+	SourceControlProvider.Execute(RevertOperation, Filenames);
+
+	// Removes selected Uncontrolled Files from their Uncontrolled Changelists
+	for (const auto& Pair : UncontrolledChangelistsStateCache)
+	{
+		const FUncontrolledChangelistStateRef& UncontrolledChangelistState = Pair.Value;
+		UncontrolledChangelistState->RemoveFiles(InUncontrolledFileStates);
+	}
+
+	Algo::Transform(InUncontrolledFileStates, Filenames, [](const FSourceControlStateRef& State) { return State->GetFilename(); });
+
+	// Add all files to their UncontrolledChangelist
+	(*ChangelistState)->AddFiles(Filenames, FUncontrolledChangelistState::ECheckFlags::None);
 
 	SaveState();
 }
 
+void FUncontrolledChangelistsModule::MoveFilesToControlledChangelist(const TArray<FSourceControlStateRef>& InUncontrolledFileStates, const FSourceControlChangelistPtr& InChangelist)
+{
+	if (!IsEnabled())
+	{
+		return;
+	}
+
+	ISourceControlProvider& SourceControlProvider = ISourceControlModule::Get().GetProvider();
+	TArray<FString> UncontrolledFilenames;
+	TArray<FSourceControlStateRef> UpdatedFilestates;
+
+	// Get updated filestates to check Checkout capabilities.
+	Algo::Transform(InUncontrolledFileStates, UncontrolledFilenames, [](const FSourceControlStateRef& State) { return State->GetFilename(); });
+	SourceControlProvider.GetState(UncontrolledFilenames, UpdatedFilestates, EStateCacheUsage::ForceUpdate);
+
+	TArray<UPackage*> PackageConflicts;
+
+	// Check if we can Checkout files
+	for (const FSourceControlStateRef& Filestate : UpdatedFilestates)
+	{
+		if (!Filestate->CanCheckout())
+		{
+			FString PackageName;
+			FPackageName::TryConvertFilenameToLongPackageName(Filestate->GetFilename(), PackageName);
+			UPackage* FoundPackage = FindOrConstructDynamicTypePackage(*PackageName);
+
+			if (FoundPackage != nullptr)
+			{
+				PackageConflicts.Add(FoundPackage);
+			}
+		}
+	}
+
+	bool bCanProceed = true;
+
+	// If we detected conflict, asking user if we should proceed.
+	if (!PackageConflicts.IsEmpty())
+	{
+		bCanProceed = ShowConflictDialog(PackageConflicts);
+	}
+
+	if (bCanProceed)
+	{
+		SourceControlProvider.Execute(ISourceControlOperation::Create<FCheckOut>(), InChangelist, UncontrolledFilenames);
+	}
+}
+
+bool FUncontrolledChangelistsModule::ShowConflictDialog(TArray<UPackage*> InPackageConflicts)
+{
+	FPackagesDialogModule& CheckoutPackagesDialogModule = FModuleManager::LoadModuleChecked<FPackagesDialogModule>(TEXT("PackagesDialog"));
+
+	CheckoutPackagesDialogModule.CreatePackagesDialog(LOCTEXT("CheckoutPackagesDialogTitle", "Check Out Assets"),
+													  LOCTEXT("CheckoutPackagesDialogMessage", "Conflict detected in the following assets:"),
+													  true);
+
+	CheckoutPackagesDialogModule.AddButton(DRT_CheckOut, LOCTEXT("Dlg_CheckOutButton", "Check Out"), LOCTEXT("Dlg_CheckOutTooltip", "Attempt to Check Out Assets"));
+	CheckoutPackagesDialogModule.AddButton(DRT_Cancel, LOCTEXT("Dlg_Cancel", "Cancel"), LOCTEXT("Dlg_CancelTooltip", "Cancel Request"));
+
+	CheckoutPackagesDialogModule.SetWarning(LOCTEXT("CheckoutPackagesWarnMessage", "Warning: These assets are locked or not at the head revision. You may lose your changes if you continue, as you will be unable to submit them to source control."));
+
+	for (UPackage* Conflict : InPackageConflicts)
+	{
+		CheckoutPackagesDialogModule.AddPackageItem(Conflict, ECheckBoxState::Undetermined);
+	}
+
+	EDialogReturnType UserResponse = CheckoutPackagesDialogModule.ShowPackagesDialog();
+
+	return UserResponse == DRT_CheckOut;
+}
+
 void FUncontrolledChangelistsModule::SaveState() const
 {
 	TSharedPtr<FJsonObject> RootObject = MakeShareable(new FJsonObject);
diff --git a/Engine/Source/Developer/UncontrolledChangelists/Public/UncontrolledChangelistState.h b/Engine/Source/Developer/UncontrolledChangelists/Public/UncontrolledChangelistState.h
--- a/Engine/Source/Developer/UncontrolledChangelists/Public/UncontrolledChangelistState.h
+++ b/Engine/Source/Developer/UncontrolledChangelists/Public/UncontrolledChangelistState.h
@@ -16,8 +16,13 @@ public:
 
 	enum class ECheckFlags
 	{
+		/** No Check */
 		None			= 0,
+
+		/** File has been modified */
 		Modified		= 1,
+
+		/** File is not checked out */
 		NotCheckedOut	= 1 << 1
 	};
 
@@ -76,12 +81,18 @@ public:
 	bool Deserialize(const TSharedRef<FJsonObject> InJsonValue);
 
 	/**
-	 * Adds a file to this Uncontrolled Changelist State.
-	 * @param 	InFilename 		The file to be added.
-	 * @param 	bInCheckStatus 	Tells whether we should verify if the file has been modified.
+	 * Adds files to this Uncontrolled Changelist State.
+	 * @param 	InFilenames		The files to be added.
+	 * @param 	InCheckFlags 	Tells which checks have to pass to add a file.
 	 */	
 	void AddFiles(const TArray<FString>& InFilenames, const ECheckFlags InCheckFlags);
 
+	/**
+	 * Removes files from this Uncontrolled Changelist State if present.
+	 * @param 	InFileStates 	The files to be removed.
+	 */
+	void RemoveFiles(const TArray<FSourceControlStateRef>& InFileStates);
+
 	/**
 	 * Updates the status of all files.
 	 */
diff --git a/Engine/Source/Developer/UncontrolledChangelists/Public/UncontrolledChangelistsModule.h b/Engine/Source/Developer/UncontrolledChangelists/Public/UncontrolledChangelistsModule.h
--- a/Engine/Source/Developer/UncontrolledChangelists/Public/UncontrolledChangelistsModule.h
+++ b/Engine/Source/Developer/UncontrolledChangelists/Public/UncontrolledChangelistsModule.h
@@ -14,7 +14,7 @@
  */
 class UNCONTROLLEDCHANGELISTS_API FUncontrolledChangelistsModule : public IModuleInterface
 {
-	typedef TMap<FUncontrolledChangelist, TSharedRef<class FUncontrolledChangelistState, ESPMode::ThreadSafe>> FUncontrolledChangelistsStateCache;
+	typedef TMap<FUncontrolledChangelist, FUncontrolledChangelistStateRef> FUncontrolledChangelistsStateCache;
 
 public:
 	static constexpr const TCHAR* VERSION_NAME = TEXT("version");
@@ -78,13 +78,19 @@ public:
 	void OnObjectTransacted(UObject* InObject, const class FTransactionObjectEvent& InTransactionEvent);
 
 	/**
-	 * Called after an action moving files to an Uncontrolled Changelist (example: drag and drop).
-	 * Moves the files to the provided uncontrolled changelist.
-	 * @param 	InFilenames 	The files to move.
-	 * @param 	InChangelist 	The Uncontrolled Changelist where to move the files.
-	 * @return 	returnDesc
+	 * Moves files to an Uncontrolled Changelist.
+	 * @param 	InControlledFileStates 		The Controlled files to move.
+	 * @param 	InUncontrolledFileStates 	The Uncontrolled files to move.
+	 * @param 	InChangelist 				The Uncontrolled Changelist where to move the files.
 	 */
-	void OnFilesMovedToUncontrolledChangelist(const TArray<FString> InFilenames, const FUncontrolledChangelist& InChangelist);
+	void MoveFilesToUncontrolledChangelist(const TArray<FSourceControlStateRef>& InControlledFileStates, const TArray<FSourceControlStateRef>& InUncontrolledFileStates, const FUncontrolledChangelist& InUncontrolledChangelist);
+
+	/**
+	 * Moves files to a Controlled Changelist.
+	 * @param 	InUncontrolledFileStates 	The files to move.
+	 * @param 	InChangelist 				The Controlled Changelist where to move the files.
+	 */
+	void MoveFilesToControlledChangelist(const TArray<FSourceControlStateRef>& InUncontrolledFileStates, const FSourceControlChangelistPtr& InChangelist);
 
 private:
 	/**
@@ -110,6 +116,13 @@ private:
 	 */
 	FString GetUObjectPackageFullpath(const UObject* InObject) const;
 
+	/**
+	 * Displays a Package Dialog warning the user about conflicting packages. 
+	 * @param 	InPackageConflicts 	The conflicting packages to display.
+	 * @return 	True if the user decided to proceed. False if they cancelled.
+	 */
+	bool ShowConflictDialog(TArray<UPackage*> InPackageConflicts);
+
 private:
 	FUncontrolledChangelistsStateCache UncontrolledChangelistsStateCache;
 	FDelegateHandle OnAssetAddedDelegateHandle;
diff --git a/Engine/Source/Developer/UncontrolledChangelists/UncontrolledChangelists.Build.cs b/Engine/Source/Developer/UncontrolledChangelists/UncontrolledChangelists.Build.cs
--- a/Engine/Source/Developer/UncontrolledChangelists/UncontrolledChangelists.Build.cs
+++ b/Engine/Source/Developer/UncontrolledChangelists/UncontrolledChangelists.Build.cs
@@ -37,6 +37,7 @@ public class UncontrolledChangelists : ModuleRules
 					"EditorFramework",
 					"Engine",
 					"UnrealEd",
+					"PackagesDialog",
 				}
 			);
         }
diff --git a/Engine/Source/Editor/SourceControlWindows/Private/SSourceControlChangelists.cpp b/Engine/Source/Editor/SourceControlWindows/Private/SSourceControlChangelists.cpp
--- a/Engine/Source/Editor/SourceControlWindows/Private/SSourceControlChangelists.cpp
+++ b/Engine/Source/Editor/SourceControlWindows/Private/SSourceControlChangelists.cpp
@@ -53,10 +53,13 @@ struct FSCCFileDragDropOp : public FDragDropOperation
 
 	virtual TSharedPtr<SWidget> GetDefaultDecorator() const override
 	{
-		return SSourceControlCommon::GetSCCFileWidget(Files[0]);
+		FSourceControlStateRef FileState = Files.IsEmpty() ? UncontrolledFiles[0] : Files[0];
+
+		return SSourceControlCommon::GetSCCFileWidget(MoveTemp(FileState));
 	}
 
 	TArray<FSourceControlStateRef> Files;
+	TArray<FSourceControlStateRef> UncontrolledFiles;
 };
 
 //////////////////////////////
@@ -1210,7 +1213,10 @@ protected:
 			Algo::Transform(Operation->Files, Files, [](const FSourceControlStateRef& State) { return State->GetFilename(); });
 
 			ISourceControlProvider& SourceControlProvider = ISourceControlModule::Get().GetProvider();
+			FUncontrolledChangelistsModule& UncontrolledChangelistModule = FUncontrolledChangelistsModule::Get();
+			
 			SourceControlProvider.Execute(ISourceControlOperation::Create<FMoveToChangelist>(), Changelist, Files);
+			UncontrolledChangelistModule.MoveFilesToControlledChangelist(Operation->UncontrolledFiles, Changelist);
 		}
 
 		return FReply::Handled();
@@ -1320,16 +1326,7 @@ protected:
 		
 		if (Operation.IsValid())
 		{
-			TArray<FString> Files;
-			Algo::Transform(Operation->Files, Files, [](const FSourceControlStateRef& State) { return State->GetFilename(); });
-
-			ISourceControlProvider& SourceControlProvider = ISourceControlModule::Get().GetProvider();
-			auto RevertOperation = ISourceControlOperation::Create<FRevert>();
-
-			RevertOperation->SetSoftRevert(true);
-			SourceControlProvider.Execute(RevertOperation, Files);
-
-			FUncontrolledChangelistsModule::Get().OnFilesMovedToUncontrolledChangelist(Files, TreeItem->UncontrolledChangelistState->Changelist);
+			FUncontrolledChangelistsModule::Get().MoveFilesToUncontrolledChangelist(Operation->Files, Operation->UncontrolledFiles, TreeItem->UncontrolledChangelistState->Changelist);
 		}
 
 		return FReply::Handled();
@@ -1562,10 +1559,27 @@ FReply SSourceControlChangelistsWidget::OnFilesDragged(const FGeometry& InGeomet
 	{
 		TSharedRef<FSCCFileDragDropOp> Operation = MakeShareable(new FSCCFileDragDropOp());
 
-		Algo::Transform(TreeView->GetSelectedItems(), Operation->Files, [](FChangelistTreeItemPtr InTreeItem) { check(InTreeItem->GetTreeItemType() == IChangelistTreeItem::File); return static_cast<FFileTreeItem*>(InTreeItem.Get())->FileState; });
+		for (FChangelistTreeItemPtr InTreeItem : TreeView->GetSelectedItems())
+		{
+			if (InTreeItem->GetTreeItemType() == IChangelistTreeItem::File)
+			{
+				FFileTreeItemRef FileTreeItem = StaticCastSharedRef<FFileTreeItem>(InTreeItem.ToSharedRef());
+				FSourceControlStateRef FileState = FileTreeItem->FileState;
+
+				if (FileTreeItem->GetParent()->GetTreeItemType() == IChangelistTreeItem::UncontrolledChangelist)
+				{
+					Operation->UncontrolledFiles.Add(MoveTemp(FileState));
+				}
+				else
+				{
+					Operation->Files.Add(MoveTemp(FileState));
+				}
+			}
+		}
+		
 		Operation->Construct();
 
-		return FReply::Handled().BeginDragDrop(Operation);
+ 		return FReply::Handled().BeginDragDrop(Operation);
 	}
 
 	return FReply::Unhandled();
diff --git a/Engine/Source/Editor/SourceControlWindows/Private/SSourceControlCommon.h b/Engine/Source/Editor/SourceControlWindows/Private/SSourceControlCommon.h
--- a/Engine/Source/Editor/SourceControlWindows/Private/SSourceControlCommon.h
+++ b/Engine/Source/Editor/SourceControlWindows/Private/SSourceControlCommon.h
@@ -95,6 +95,9 @@ struct FUncontrolledChangelistTreeItem : public IChangelistTreeItem
 	FUncontrolledChangelistStateRef UncontrolledChangelistState;
 };
 
+typedef TSharedPtr<FUncontrolledChangelistTreeItem> FUncontrolledChangelistTreeItemPtr;
+typedef TSharedRef<FUncontrolledChangelistTreeItem> FUncontrolledChangelistTreeItemRef;
+
 struct FShelvedChangelistTreeItem : public IChangelistTreeItem
 {
 	FShelvedChangelistTreeItem()
@@ -185,6 +188,9 @@ private:
 	TArray<FAssetData> Assets;
 };
 
+typedef TSharedPtr<FFileTreeItem> FFileTreeItemPtr;
+typedef TSharedRef<FFileTreeItem> FFileTreeItemRef;
+
 struct FShelvedFileTreeItem : public FFileTreeItem
 {
 	explicit FShelvedFileTreeItem(FSourceControlStateRef InFileState, bool bBeautifyPaths = true)
