Adds automatic discovery of potentially modified files out of SourceControl.
* Adds reconcile contextual actions to manually trigger the verification process
* Loads UncontrolledChangelists module along UnrealEd to be able to hook loaded/saved files sooner
* Adds CPU Trace for FPerforceUpdateStatusWorker::Execute and ParseUpdateStatusResults
* Changes FPerforceCheckOutWorker to update Changelist State if checked out directly to a CL
* Changes FPerforceRevertWorker to update Changelist State after revert

#changelist validated

diff --git a/Engine/Plugins/Developer/PerforceSourceControl/Source/PerforceSourceControl/Private/PerforceSourceControlOperations.cpp b/Engine/Plugins/Developer/PerforceSourceControl/Source/PerforceSourceControl/Private/PerforceSourceControlOperations.cpp
--- a/Engine/Plugins/Developer/PerforceSourceControl/Source/PerforceSourceControl/Private/PerforceSourceControlOperations.cpp
+++ b/Engine/Plugins/Developer/PerforceSourceControl/Source/PerforceSourceControl/Private/PerforceSourceControlOperations.cpp
@@ -323,6 +323,7 @@ bool FPerforceCheckOutWorker::Execute(FPerforceSourceControlCommand& InCommand)
 			{
 				Parameters.Add(TEXT("-c"));
 				Parameters.Add(ChangelistNumber);
+				InChangelist = InCommand.Changelist;
 			}
 		}
 
@@ -337,6 +338,24 @@ bool FPerforceCheckOutWorker::Execute(FPerforceSourceControlCommand& InCommand)
 
 bool FPerforceCheckOutWorker::UpdateStates() const
 {
+	// If files have been checkedout directly to a CL, modify the cached state to reflect it.
+	if (InChangelist.IsInitialized())
+	{
+		FPerforceSourceControlModule& PerforceSourceControl = FPerforceSourceControlModule::Get();
+		TSharedRef<FPerforceSourceControlChangelistState, ESPMode::ThreadSafe> ChangelistState = PerforceSourceControl.GetProvider().GetStateInternal(InChangelist);
+
+		for (TMap<FString, EPerforceState::Type>::TConstIterator It(OutResults); It; ++It)
+		{
+			if (It.Value() != EPerforceState::CheckedOut)
+			{
+				continue;
+			}
+
+			TSharedRef<FPerforceSourceControlState, ESPMode::ThreadSafe> State = PerforceSourceControl.GetProvider().GetStateInternal(It.Key());
+			ChangelistState->Files.Add(State);
+		}
+	}
+
 	return UpdateCachedStates(OutResults);
 }
 
@@ -702,8 +721,20 @@ bool FPerforceRevertWorker::Execute(FPerforceSourceControlCommand& InCommand)
 
 bool FPerforceRevertWorker::UpdateStates() const
 {
+	FPerforceSourceControlModule& PerforceSourceControl = FPerforceSourceControlModule::Get();
 	bool bUpdatedCachedStates = UpdateCachedStates(OutResults);
-	bool bUpdatedChangelists = ChangelistToUpdate.IsInitialized() && RemoveFilesFromChangelist(OutResults, ChangelistToUpdate);
+	bool bUpdatedChangelists = false;
+
+	bUpdatedChangelists = ChangelistToUpdate.IsInitialized() && RemoveFilesFromChangelist(OutResults, ChangelistToUpdate);
+
+	// Use reverted files to update changelist state.
+	for (TMap<FString, EPerforceState::Type>::TConstIterator It(OutResults); It; ++It)
+	{
+		TSharedRef<FPerforceSourceControlState, ESPMode::ThreadSafe> State = PerforceSourceControl.GetProvider().GetStateInternal(It.Key());
+
+		bUpdatedChangelists |= State->Changelist.IsInitialized() && RemoveFilesFromChangelist(OutResults, State->Changelist);
+	}
+	
 	return bUpdatedCachedStates || bUpdatedChangelists;
 }
 
@@ -884,6 +915,7 @@ static FSourceControlDigest ComputeLocalDigest(FString InFilename, FString InDep
 
 static void ParseUpdateStatusResults(const FP4RecordSet& InRecords, const TArray<FText>& ErrorMessages, TArray<FPerforceSourceControlState>& OutStates, TArray<FString>& OutModifiedFiles, const FString& ContentRoot, TMap<FString, FBranchModification>& BranchModifications)
 {
+	TRACE_CPUPROFILER_EVENT_SCOPE(FPerforceUpdateStatusWorker_ParseUpdateStatusResults);
 	// Build up a map of any other branch states
 	for (int32 Index = 0; Index < InRecords.Num(); ++Index)
 	{
@@ -1534,6 +1566,7 @@ FName FPerforceUpdateStatusWorker::GetName() const
 
 bool FPerforceUpdateStatusWorker::Execute(FPerforceSourceControlCommand& InCommand)
 {
+	TRACE_CPUPROFILER_EVENT_SCOPE(FPerforceUpdateStatusWorker::Execute);
 	FScopedPerforceConnection ScopedConnection(InCommand);
 	if (!InCommand.IsCanceled() && ScopedConnection.IsValid())
 	{
diff --git a/Engine/Plugins/Developer/PerforceSourceControl/Source/PerforceSourceControl/Private/PerforceSourceControlOperations.h b/Engine/Plugins/Developer/PerforceSourceControl/Source/PerforceSourceControl/Private/PerforceSourceControlOperations.h
--- a/Engine/Plugins/Developer/PerforceSourceControl/Source/PerforceSourceControl/Private/PerforceSourceControlOperations.h
+++ b/Engine/Plugins/Developer/PerforceSourceControl/Source/PerforceSourceControl/Private/PerforceSourceControlOperations.h
@@ -32,6 +32,7 @@ public:
 public:
 	/** Map of filenames to perforce state */
 	TMap<FString, EPerforceState::Type> OutResults;
+	FPerforceSourceControlChangelist InChangelist;
 };
 
 class FPerforceCheckInWorker : public IPerforceSourceControlWorker
diff --git a/Engine/Source/Developer/UncontrolledChangelists/Private/UncontrolledChangelistState.cpp b/Engine/Source/Developer/UncontrolledChangelists/Private/UncontrolledChangelistState.cpp
--- a/Engine/Source/Developer/UncontrolledChangelists/Private/UncontrolledChangelistState.cpp
+++ b/Engine/Source/Developer/UncontrolledChangelists/Private/UncontrolledChangelistState.cpp
@@ -3,7 +3,6 @@
 #include "UncontrolledChangelistState.h"
 
 #include "Algo/Transform.h"
-#include "Algo/ForEach.h"
 #include "Dom/JsonObject.h"
 #include "Dom/JsonValue.h"
 #include "ISourceControlModule.h"
@@ -88,16 +87,17 @@ bool FUncontrolledChangelistState::Deserialize(const TSharedRef<FJsonObject> InJ
 	return true;
 }
 
-void FUncontrolledChangelistState::AddFiles(const TArray<FString>& InFilenames, const ECheckFlags InCheckFlags)
+bool FUncontrolledChangelistState::AddFiles(const TArray<FString>& InFilenames, const ECheckFlags InCheckFlags)
 {
 	TArray<FSourceControlStateRef> FileStates;
 	ISourceControlProvider& SourceControlProvider = ISourceControlModule::Get().GetProvider();
 	bool bCheckStatus = (InCheckFlags & ECheckFlags::Modified) != ECheckFlags::None;
 	bool bCheckCheckout = (InCheckFlags & ECheckFlags::NotCheckedOut) != ECheckFlags::None;
+	bool bOutChanged = false;
 
 	if (InFilenames.IsEmpty())
 	{
-		return;
+		return bOutChanged;
 	}
 
 	if (bCheckStatus)
@@ -119,35 +119,48 @@ void FUncontrolledChangelistState::AddFiles(const TArray<FString>& InFilenames,
 			if (bIsCheckoutCompliant && bIsStatusCompliant)
 			{
 				Files.Add(FileState);
+				bOutChanged = true;
 			}
 		}
 	}
+
+	return bOutChanged;
 }
 
-void FUncontrolledChangelistState::RemoveFiles(const TArray<FSourceControlStateRef>& InFileStates)
+bool FUncontrolledChangelistState::RemoveFiles(const TArray<FSourceControlStateRef>& InFileStates)
 {
-	Algo::ForEach(InFileStates, [&Files = Files](const FSourceControlStateRef& FileState) { Files.Remove(FileState); });
+	bool bOutChanged = false;
+
+	for (const FSourceControlStateRef& FileState : InFileStates)
+	{
+		int32 OldSize = Files.Num();
+
+		Files.Remove(FileState);
+
+		bOutChanged |= OldSize != Files.Num();
+	}
+
+	return bOutChanged;
 }
 
-void FUncontrolledChangelistState::UpdateStatus()
+bool FUncontrolledChangelistState::UpdateStatus()
 {
 	TArray<FString> FilesToUpdate;
+	bool bOutChanged = false;
+
 	Algo::Transform(Files, FilesToUpdate, [](const FSourceControlStateRef& State) { return State->GetFilename(); });
 
 	if (FilesToUpdate.Num() == 0)
 	{
-		return;
+		return bOutChanged;
 	}
 
 	ISourceControlProvider& SourceControlProvider = ISourceControlModule::Get().GetProvider();
 	auto UpdateStatusOperation = ISourceControlOperation::Create<FUpdateStatus>();
 	UpdateStatusOperation->SetUpdateModifiedState(true);
 
-	SourceControlProvider.Execute(UpdateStatusOperation, FilesToUpdate, EConcurrency::Asynchronous, FSourceControlOperationComplete::CreateSP(this, &FUncontrolledChangelistState::OnUpdateStatusCompleted));
-}
+	SourceControlProvider.Execute(UpdateStatusOperation, FilesToUpdate);
 
-void FUncontrolledChangelistState::OnUpdateStatusCompleted(const FSourceControlOperationRef& InOperation, ECommandResult::Type InResult)
-{
 	for (auto It = Files.CreateIterator(); It; ++It)
 	{
 		FSourceControlStateRef& State = *It;
@@ -155,8 +168,22 @@ void FUncontrolledChangelistState::OnUpdateStatusCompleted(const FSourceControlO
 		if (State->IsCheckedOut() || (!State->IsModified()))
 		{
 			It.RemoveCurrent();
+			bOutChanged = true;
 		}
 	}
+
+	return bOutChanged;
+}
+
+void FUncontrolledChangelistState::RemoveDuplicates(TSet<FString>& InOutLoadedFiles, TSet<FString>& InOutModifiedFiles)
+{
+	for (const FSourceControlStateRef& FileState : Files)
+	{
+		const FString& Filename = FileState->GetFilename();
+		
+		InOutLoadedFiles.Remove(Filename);
+		InOutModifiedFiles.Remove(Filename);
+	}
 }
 
 #undef LOCTEXT_NAMESPACE
diff --git a/Engine/Source/Developer/UncontrolledChangelists/Private/UncontrolledChangelistsModule.cpp b/Engine/Source/Developer/UncontrolledChangelists/Private/UncontrolledChangelistsModule.cpp
--- a/Engine/Source/Developer/UncontrolledChangelists/Private/UncontrolledChangelistsModule.cpp
+++ b/Engine/Source/Developer/UncontrolledChangelists/Private/UncontrolledChangelistsModule.cpp
@@ -2,6 +2,7 @@
 
 #include "UncontrolledChangelistsModule.h"
 
+#include "Algo/Copy.h"
 #include "Algo/Find.h"
 #include "Algo/Transform.h"
 #include "AssetRegistryModule.h"
@@ -17,6 +18,7 @@
 #include "Serialization/JsonWriter.h"
 #include "SourceControlOperations.h"
 #include "Styling/SlateTypes.h"
+#include "UObject/ObjectSaveContext.h"
 
 #define LOCTEXT_NAMESPACE "UncontrolledChangelists"
 
@@ -37,10 +39,7 @@ void FUncontrolledChangelistsModule::StartupModule()
 	FAssetRegistryModule& AssetRegistryModule = FModuleManager::GetModuleChecked<FAssetRegistryModule>(TEXT("AssetRegistry"));
 	OnAssetAddedDelegateHandle = AssetRegistryModule.Get().OnAssetAdded().AddLambda([](const struct FAssetData& AssetData) { Get().OnAssetAdded(AssetData); });
 	OnAssetLoadedDelegateHandle = FCoreUObjectDelegates::OnAssetLoaded.AddLambda([](UObject* InAsset) { Get().OnAssetLoaded(InAsset); });
-	OnObjectTransactedDelegateHandle = FCoreUObjectDelegates::OnObjectTransacted.AddLambda([](UObject* InObject, const class FTransactionObjectEvent& InTransactionEvent)
-	{
-		Get().OnObjectTransacted(InObject, InTransactionEvent);
-	});
+	OnObjectPreSavedDelegateHandle = FCoreUObjectDelegates::OnObjectPreSave.AddLambda([](UObject* InAsset, const FObjectPreSaveContext& InPreSaveContext) { Get().OnObjectPreSaved(InAsset, InPreSaveContext); });
 }
 
 void FUncontrolledChangelistsModule::ShutdownModule()
@@ -54,7 +53,7 @@ void FUncontrolledChangelistsModule::ShutdownModule()
 	}
 
 	FCoreUObjectDelegates::OnAssetLoaded.Remove(OnAssetLoadedDelegateHandle);
-	FCoreUObjectDelegates::OnObjectTransacted.Remove(OnObjectTransactedDelegateHandle);
+	FCoreUObjectDelegates::OnObjectPreSave.Remove(OnObjectPreSavedDelegateHandle);
 }
 
 bool FUncontrolledChangelistsModule::IsEnabled() const
@@ -76,6 +75,8 @@ TArray<FUncontrolledChangelistStateRef> FUncontrolledChangelistsModule::GetChang
 
 void FUncontrolledChangelistsModule::OnMakeWritable(const FString& InFilename)
 {
+	bool bHasStateChanged = false;
+
 	if (!IsEnabled())
 	{
 		return;
@@ -84,13 +85,18 @@ void FUncontrolledChangelistsModule::OnMakeWritable(const FString& InFilename)
 	FUncontrolledChangelist DefaultUncontrolledChangelist(FUncontrolledChangelist::DEFAULT_UNCONTROLLED_CHANGELIST_GUID, FUncontrolledChangelist::DEFAULT_UNCONTROLLED_CHANGELIST_NAME.ToString());
 	FUncontrolledChangelistsStateCache::ValueType& UncontrolledChangelistState = UncontrolledChangelistsStateCache.FindOrAdd(MoveTemp(DefaultUncontrolledChangelist), MakeShareable(new FUncontrolledChangelistState(DefaultUncontrolledChangelist)));
 
-	UncontrolledChangelistState->AddFiles({ InFilename }, FUncontrolledChangelistState::ECheckFlags::NotCheckedOut);
+	bHasStateChanged = UncontrolledChangelistState->AddFiles({ InFilename }, FUncontrolledChangelistState::ECheckFlags::NotCheckedOut);
 
-	SaveState();
+	if (bHasStateChanged)
+	{
+		OnStateChanged();
+	}
 }
 
 void FUncontrolledChangelistsModule::UpdateStatus()
 {
+	bool bHasStateChanged = false;
+
 	if (!IsEnabled())
 	{
 		return;
@@ -99,10 +105,74 @@ void FUncontrolledChangelistsModule::UpdateStatus()
 	for (FUncontrolledChangelistsStateCache::ElementType& Pair : UncontrolledChangelistsStateCache)
 	{
 		FUncontrolledChangelistsStateCache::ValueType& UncontrolledChangelistState = Pair.Value;
-		UncontrolledChangelistState->UpdateStatus();
+		bHasStateChanged |= UncontrolledChangelistState->UpdateStatus();
+	}
+
+	if (bHasStateChanged)
+	{
+		OnStateChanged();
 	}
 }
 
+FText FUncontrolledChangelistsModule::GetReconcileStatus() const
+{
+	return FText::Format(LOCTEXT("ReconcileStatus", "Assets to check for reconcile: {0} loaded, {1} saved"), FText::AsNumber(LoadedFilesCache.Num()), FText::AsNumber(SavedFilesCache.Num()));
+}
+
+void FUncontrolledChangelistsModule::OnReconcileLoadedAssets()
+{
+	if ((!IsEnabled()) || ((LoadedFilesCache.Num() == 0)))
+	{
+		return;
+	}
+
+	CleanAssetsCaches();
+	AddFilesToDefaultUncontrolledChangelist(LoadedFilesCache.Array(), FUncontrolledChangelistState::ECheckFlags::All);
+	LoadedFilesCache.Reset();
+}
+
+void FUncontrolledChangelistsModule::OnReconcileSavedAssets()
+{
+	if ((!IsEnabled()) || ((SavedFilesCache.Num() == 0)))
+	{
+		return;
+	}
+
+	CleanAssetsCaches();
+	AddFilesToDefaultUncontrolledChangelist(SavedFilesCache.Array(), FUncontrolledChangelistState::ECheckFlags::All);
+	SavedFilesCache.Reset();
+}
+
+void FUncontrolledChangelistsModule::OnReconcileAllAssets()
+{
+	if ((!IsEnabled()) || ((LoadedFilesCache.Num() == 0) && (SavedFilesCache.Num() == 0)))
+	{
+		return;
+	}
+
+	CleanAssetsCaches();
+
+	TArray<FString> FilenamesToCheck;
+
+	Algo::Copy(LoadedFilesCache, FilenamesToCheck);
+	Algo::Copy(SavedFilesCache, FilenamesToCheck);
+
+	AddFilesToDefaultUncontrolledChangelist(FilenamesToCheck, FUncontrolledChangelistState::ECheckFlags::All);
+
+	LoadedFilesCache.Reset();
+	SavedFilesCache.Reset();
+}
+
+void FUncontrolledChangelistsModule::OnClearLoadedAssetsCache()
+{
+	LoadedFilesCache.Reset();
+}
+
+void FUncontrolledChangelistsModule::OnClearSavedAssetsCache()
+{
+	SavedFilesCache.Reset();
+}
+
 void FUncontrolledChangelistsModule::OnAssetAdded(const struct FAssetData& AssetData)
 {
 
@@ -110,26 +180,42 @@ void FUncontrolledChangelistsModule::OnAssetAdded(const struct FAssetData& Asset
 
 void FUncontrolledChangelistsModule::OnAssetLoaded(UObject* InAsset)
 {
+	if (!IsEnabled())
+	{
+		return;
+	}
+
 	FString Fullpath = GetUObjectPackageFullpath(InAsset);
 
 	if (Fullpath.IsEmpty())
 	{
 		return;
 	}
+
+	LoadedFilesCache.Add(MoveTemp(Fullpath));
 }
 
-void FUncontrolledChangelistsModule::OnObjectTransacted(UObject* InObject, const class FTransactionObjectEvent& InTransactionEvent)
+void FUncontrolledChangelistsModule::OnObjectPreSaved(UObject* InAsset, const FObjectPreSaveContext& InPreSaveContext)
 {
-	FString Fullpath = GetUObjectPackageFullpath(InObject);
+	if (!IsEnabled())
+	{
+		return;
+	}
+
+	FString Fullpath = FPaths::ConvertRelativePathToFull(InPreSaveContext.GetTargetFilename());
 
 	if (Fullpath.IsEmpty())
 	{
 		return;
 	}
+
+	SavedFilesCache.Add(MoveTemp(Fullpath));
 }
 
 void FUncontrolledChangelistsModule::MoveFilesToUncontrolledChangelist(const TArray<FSourceControlStateRef>& InControlledFileStates, const TArray<FSourceControlStateRef>& InUncontrolledFileStates, const FUncontrolledChangelist& InUncontrolledChangelist)
 {
+	bool bHasStateChanged = false;
+
 	if (!IsEnabled())
 	{
 		return;
@@ -162,9 +248,12 @@ void FUncontrolledChangelistsModule::MoveFilesToUncontrolledChangelist(const TAr
 	Algo::Transform(InUncontrolledFileStates, Filenames, [](const FSourceControlStateRef& State) { return State->GetFilename(); });
 
 	// Add all files to their UncontrolledChangelist
-	(*ChangelistState)->AddFiles(Filenames, FUncontrolledChangelistState::ECheckFlags::None);
+	bHasStateChanged = (*ChangelistState)->AddFiles(Filenames, FUncontrolledChangelistState::ECheckFlags::None);
 
-	SaveState();
+	if (bHasStateChanged)
+	{
+		OnStateChanged();
+	}
 }
 
 void FUncontrolledChangelistsModule::MoveFilesToControlledChangelist(const TArray<FSourceControlStateRef>& InUncontrolledFileStates, const FSourceControlChangelistPtr& InChangelist)
@@ -211,6 +300,9 @@ void FUncontrolledChangelistsModule::MoveFilesToControlledChangelist(const TArra
 	if (bCanProceed)
 	{
 		SourceControlProvider.Execute(ISourceControlOperation::Create<FCheckOut>(), InChangelist, UncontrolledFilenames);
+
+		// UpdateStatus so UncontrolledChangelists can remove files from their cache if they were present before checkout.
+		UpdateStatus();
 	}
 }
 
@@ -237,6 +329,38 @@ bool FUncontrolledChangelistsModule::ShowConflictDialog(TArray<UPackage*> InPack
 	return UserResponse == DRT_CheckOut;
 }
 
+void FUncontrolledChangelistsModule::OnStateChanged()
+{
+	OnUncontrolledChangelistModuleChanged.Broadcast();
+	SaveState();
+}
+
+void FUncontrolledChangelistsModule::CleanAssetsCaches()
+{
+	// Remove files we are already tracking in Uncontrolled Changelists
+	for (FUncontrolledChangelistsStateCache::ElementType& Pair : UncontrolledChangelistsStateCache)
+	{
+		FUncontrolledChangelistsStateCache::ValueType& UncontrolledChangelistState = Pair.Value;
+		UncontrolledChangelistState->RemoveDuplicates(LoadedFilesCache, SavedFilesCache);
+	}
+}
+
+void FUncontrolledChangelistsModule::AddFilesToDefaultUncontrolledChangelist(const TArray<FString>& InFilenames, const FUncontrolledChangelistState::ECheckFlags InCheckFlags)
+{
+	bool bHasStateChanged = false;
+
+	FUncontrolledChangelist DefaultUncontrolledChangelist(FUncontrolledChangelist::DEFAULT_UNCONTROLLED_CHANGELIST_GUID, FUncontrolledChangelist::DEFAULT_UNCONTROLLED_CHANGELIST_NAME.ToString());
+	FUncontrolledChangelistsStateCache::ValueType& UncontrolledChangelistState = UncontrolledChangelistsStateCache.FindOrAdd(MoveTemp(DefaultUncontrolledChangelist), MakeShareable(new FUncontrolledChangelistState(DefaultUncontrolledChangelist)));
+
+	// Try to add files, they will be added only if they pass the required checks
+	bHasStateChanged = UncontrolledChangelistState->AddFiles(InFilenames, InCheckFlags);
+
+	if (bHasStateChanged)
+	{
+		OnStateChanged();
+	}
+}
+
 void FUncontrolledChangelistsModule::SaveState() const
 {
 	TSharedPtr<FJsonObject> RootObject = MakeShareable(new FJsonObject);
diff --git a/Engine/Source/Developer/UncontrolledChangelists/Public/UncontrolledChangelistState.h b/Engine/Source/Developer/UncontrolledChangelists/Public/UncontrolledChangelistState.h
--- a/Engine/Source/Developer/UncontrolledChangelists/Public/UncontrolledChangelistState.h
+++ b/Engine/Source/Developer/UncontrolledChangelists/Public/UncontrolledChangelistState.h
@@ -6,7 +6,6 @@
 #include "ISourceControlChangelistState.h"
 #include "ISourceControlOperation.h"
 #include "ISourceControlProvider.h"
-
 #include "UncontrolledChangelist.h"
 
 class FUncontrolledChangelistState : public TSharedFromThis<FUncontrolledChangelistState, ESPMode::ThreadSafe>
@@ -23,7 +22,10 @@ public:
 		Modified		= 1,
 
 		/** File is not checked out */
-		NotCheckedOut	= 1 << 1
+		NotCheckedOut	= 1 << 1,
+
+		/** All the above checks */
+		All = Modified | NotCheckedOut,
 	};
 
 public:
@@ -84,34 +86,33 @@ public:
 	 * Adds files to this Uncontrolled Changelist State.
 	 * @param 	InFilenames		The files to be added.
 	 * @param 	InCheckFlags 	Tells which checks have to pass to add a file.
-	 */	
-	void AddFiles(const TArray<FString>& InFilenames, const ECheckFlags InCheckFlags);
+	 * @return 	True if a change has been performed in the Uncontrolled Changelist State.
+	 */
+	bool AddFiles(const TArray<FString>& InFilenames, const ECheckFlags InCheckFlags);
 
 	/**
 	 * Removes files from this Uncontrolled Changelist State if present.
 	 * @param 	InFileStates 	The files to be removed.
+	 * @return 	True if a change has been performed in the Uncontrolled Changelist State.
 	 */
-	void RemoveFiles(const TArray<FSourceControlStateRef>& InFileStates);
+	bool RemoveFiles(const TArray<FSourceControlStateRef>& InFileStates);
 
 	/**
-	 * Updates the status of all files.
+	 * Updates the status of all files contained in this changelist.
+	 * @return 	True if the state has been modified.
 	 */
-	void UpdateStatus();
-
-private:
+	bool UpdateStatus();
 
 	/**
-	 * Called upon completion of FStatus operation
-	 * @param 	InOperation 	The operation just completed.
-	 * @param 	InResult 		The result of the operation.
+	 * Removes files present both in the Uncontrolled Changelist and the provided sets.
+	 * @param 	InOutLoadedFiles 	A Set representing loaded files to check against.
+	 * @param 	InOutModifiedFiles 	A Set representing modified files to check against.
 	 */
-	void OnUpdateStatusCompleted(const FSourceControlOperationRef& InOperation, ECommandResult::Type InResult);
+	void RemoveDuplicates(TSet<FString>& InOutLoadedFiles, TSet<FString>& InOutModifiedFiles);
 
 public:
 	FUncontrolledChangelist Changelist;
-
 	FString Description;
-
 	TSet<FSourceControlStateRef> Files;
 
 	/** The timestamp of the last update */
diff --git a/Engine/Source/Developer/UncontrolledChangelists/Public/UncontrolledChangelistsModule.h b/Engine/Source/Developer/UncontrolledChangelists/Public/UncontrolledChangelistsModule.h
--- a/Engine/Source/Developer/UncontrolledChangelists/Public/UncontrolledChangelistsModule.h
+++ b/Engine/Source/Developer/UncontrolledChangelists/Public/UncontrolledChangelistsModule.h
@@ -16,11 +16,15 @@ class UNCONTROLLEDCHANGELISTS_API FUncontrolledChangelistsModule : public IModul
 {
 	typedef TMap<FUncontrolledChangelist, FUncontrolledChangelistStateRef> FUncontrolledChangelistsStateCache;
 
-public:
+public:	
 	static constexpr const TCHAR* VERSION_NAME = TEXT("version");
 	static constexpr const TCHAR* CHANGELISTS_NAME = TEXT("changelists");
 	static constexpr uint32 VERSION_NUMBER = 0;
 
+	/** Callback called when the state of the Uncontrolled Changelist Module (or any Uncontrolled Changelist) changed */
+	DECLARE_MULTICAST_DELEGATE(FOnUncontrolledChangelistModuleChanged);
+	FOnUncontrolledChangelistModuleChanged OnUncontrolledChangelistModuleChanged;
+
 public:
 	/** IModuleInterface implementation */
 	virtual void StartupModule() override;
@@ -36,12 +40,11 @@ public:
 	 */
 	 TArray<FUncontrolledChangelistStateRef> GetChangelistStates() const;
 
-	/**
-	 * Called when a file has been made writable. Adds the file to the Default Uncontrolled Changelist
-	 * @param	InFilename			The file to be added.
-	 */
-	void OnMakeWritable(const FString& InFilename);
-
+	 /**
+	  * Called when a file has been made writable. Adds the file to the Default Uncontrolled Changelist
+	  * @param	InFilename			The file to be added.
+	  */
+	 void OnMakeWritable(const FString& InFilename);
 
 	/**
 	 * Updates the status of Uncontrolled Changelists and files.
@@ -58,6 +61,33 @@ public:
 		return FModuleManager::LoadModuleChecked<FUncontrolledChangelistsModule>(UncontrolledChangelistsModuleName);
 	}
 
+	/**
+	 * Gets a message indicating the status of SCC coherence.
+	 * @return 	A text representing the status of SCC.
+	 */
+	FText GetReconcileStatus() const;
+
+	/** Called when "Reconcile loaded assets" button is clicked. Checks for uncontrolled modifications in previously loaded assets.
+	 *	Adds modified files to Uncontrolled Changelists
+	 */
+	void OnReconcileLoadedAssets();
+
+	/** Called when "Reconcile saved assets" button is clicked. Checks for uncontrolled modifications in previously saved assets.
+	 *	Adds modified files to Uncontrolled Changelists
+	 */
+	void OnReconcileSavedAssets();
+
+	/** Called when the "Reconcile all assets" button is clicked. Checks for uncontrolled modifications in previously loaded and saved assets.
+	 *	Adds modified files to Uncontrolled Changelists
+	 */
+	void OnReconcileAllAssets();
+
+	/** Clears the cache containing unchecked loaded assets. */
+	void OnClearLoadedAssetsCache();
+
+	/** Clears the cache containing unchecked saved assets. */
+	void OnClearSavedAssetsCache();
+
 	/**
 	 * Delegate callback called when assets are added to AssetRegistry.
 	 * @param 	AssetData 	The asset just added.
@@ -71,11 +101,11 @@ public:
 	void OnAssetLoaded(UObject* InAsset);
 
 	/**
-	 * Delegate callback called when an object is transacted.
-	 * @param 	InObject 	        The transacted object.
-	 * @param 	InTransactionEvent 	The event representing the transaction.
+	 * Delegate callback called before an asset has been written to disk.
+	 * @param 	InAsset 			The saved asset.
+	 * @param 	InPreSaveContext 	Interface used to access saved parameters.
 	 */
-	void OnObjectTransacted(UObject* InObject, const class FTransactionObjectEvent& InTransactionEvent);
+	void OnObjectPreSaved(UObject* InAsset, const FObjectPreSaveContext& InPreSaveContext);
 
 	/**
 	 * Moves files to an Uncontrolled Changelist.
@@ -123,9 +153,24 @@ private:
 	 */
 	bool ShowConflictDialog(TArray<UPackage*> InPackageConflicts);
 
+	/** Called when a state changed either in the module or an Uncontrolled Changelist. */
+	void OnStateChanged();
+
+	/** Removes from asset caches files already present in Uncontrolled Changelists */
+	void CleanAssetsCaches();
+
+	/**
+	 * Try to add the provided filenames to the default Uncontrolled Changelist.
+	 * @param 	InFilenames 	The files to add.
+	 * @param 	InCheckFlags 	The required checks to check the file against before adding.
+	 */
+	void AddFilesToDefaultUncontrolledChangelist(const TArray<FString>& InFilenames, const FUncontrolledChangelistState::ECheckFlags InCheckFlags);
+
 private:
-	FUncontrolledChangelistsStateCache UncontrolledChangelistsStateCache;
-	FDelegateHandle OnAssetAddedDelegateHandle;
-	FDelegateHandle OnAssetLoadedDelegateHandle;
-	FDelegateHandle OnObjectTransactedDelegateHandle;
+	FUncontrolledChangelistsStateCache	UncontrolledChangelistsStateCache;
+	TSet<FString>						LoadedFilesCache;
+	TSet<FString>						SavedFilesCache;
+	FDelegateHandle						OnAssetAddedDelegateHandle;
+	FDelegateHandle						OnAssetLoadedDelegateHandle;
+	FDelegateHandle						OnObjectPreSavedDelegateHandle;
 };
diff --git a/Engine/Source/Editor/SourceControlWindows/Private/SSourceControlChangelists.cpp b/Engine/Source/Editor/SourceControlWindows/Private/SSourceControlChangelists.cpp
--- a/Engine/Source/Editor/SourceControlWindows/Private/SSourceControlChangelists.cpp
+++ b/Engine/Source/Editor/SourceControlWindows/Private/SSourceControlChangelists.cpp
@@ -72,8 +72,11 @@ void SSourceControlChangelistsWidget::Construct(const FArguments& InArgs)
 {
 	// Register delegates
 	ISourceControlModule& SCCModule = ISourceControlModule::Get();
+	FUncontrolledChangelistsModule& UncontrolledChangelistModule = FUncontrolledChangelistsModule::Get();
+
 	SCCModule.RegisterProviderChanged(FSourceControlProviderChanged::FDelegate::CreateSP(this, &SSourceControlChangelistsWidget::OnSourceControlProviderChanged));
 	SourceControlStateChangedDelegateHandle = SCCModule.GetProvider().RegisterSourceControlStateChanged_Handle(FSourceControlStateChanged::FDelegate::CreateSP(this, &SSourceControlChangelistsWidget::OnSourceControlStateChanged));
+	UncontrolledChangelistModule.OnUncontrolledChangelistModuleChanged.AddSP(this, &SSourceControlChangelistsWidget::OnSourceControlStateChanged);
 
 	TreeView = CreateTreeviewWidget();
 
@@ -109,13 +112,31 @@ void SSourceControlChangelistsWidget::Construct(const FArguments& InArgs)
 		.AutoHeight()
 		[
 			SNew(SBorder)
-			.Visibility_Lambda([this]() -> EVisibility
-			{
-				return bIsRefreshing ? EVisibility::Visible : EVisibility::Collapsed;
-			})
 			[
-				SNew(STextBlock)
-				.Text_Lambda([this]() { return RefreshStatus; })
+				SNew(SHorizontalBox)
+				+ SHorizontalBox::Slot()
+				.HAlign(HAlign_Left)
+				.VAlign(VAlign_Center)
+				[
+					SNew(STextBlock)
+					.Text_Lambda([this]() { return RefreshStatus; })
+					.Visibility_Lambda([this]() -> EVisibility
+					{
+						return bIsRefreshing ? EVisibility::Visible : EVisibility::Collapsed;
+					})
+				]
+				+ SHorizontalBox::Slot()
+				.HAlign(HAlign_Right)
+				.VAlign(VAlign_Center)
+				.Padding(FMargin(2.f, 0.f))
+				[
+					SNew(STextBlock)
+					.Text_Lambda([]() { return FUncontrolledChangelistsModule::Get().GetReconcileStatus(); })
+					.Visibility_Lambda([]() -> EVisibility
+					{
+						return FUncontrolledChangelistsModule::Get().IsEnabled() ? EVisibility::Visible : EVisibility::Collapsed;
+					})
+				]
 			]
 		]
 	];
@@ -1082,6 +1103,33 @@ TSharedPtr<SWidget> SSourceControlChangelistsWidget::OnOpenContextMenu()
 				FCanExecuteAction::CreateSP(this, &SSourceControlChangelistsWidget::CanDiffAgainstDepot)));
 	}
 
+	if (FUncontrolledChangelistsModule::Get().IsEnabled())
+	{
+		Section.AddSeparator("ReconcileSeparator");
+
+		Section.AddSubMenu("ReconcileSubMenu", LOCTEXT("ReconcileSubMenu", "Reconcile"), LOCTEXT("ReconcileSubMenu_ToolTip", "Opens the Reconcile sub menu"), FNewToolMenuDelegate::CreateLambda([](UToolMenu* Menu)
+		{
+			FToolMenuSection& SubSection = Menu->AddSection(TEXT("Reconcile"));
+			
+			SubSection.AddMenuEntry("Reconcile all assets", LOCTEXT("SourceControl_ReconcileAllAssets", "Reconcile all assets"), LOCTEXT("SourceControl_ReconcileAllAssets_Tooltip", "Look for uncontrolled modification in previously loaded and saved assets."), FSlateIcon(),
+									FUIAction(FExecuteAction::CreateLambda([]() { FUncontrolledChangelistsModule::Get().OnReconcileAllAssets(); })));
+
+			SubSection.AddMenuEntry("Reconcile loaded assets", LOCTEXT("SourceControl_ReconcileLoadedAssets", "Reconcile loaded assets"), LOCTEXT("SourceControl_ReconcileLoadedAssets_Tooltip", "Look for uncontrolled modification in previously loaded assets."), FSlateIcon(),
+									FUIAction(FExecuteAction::CreateLambda([]() { FUncontrolledChangelistsModule::Get().OnReconcileLoadedAssets(); })));
+
+			SubSection.AddMenuEntry("Reconcile saved assets", LOCTEXT("SourceControl_ReconcileSavedAssets", "Reconcile saved assets"), LOCTEXT("SourceControl_ReconcileSavedAssets_Tooltip", "Look for uncontrolled modification in previously saved assets."), FSlateIcon(),
+									FUIAction(FExecuteAction::CreateLambda([]() { FUncontrolledChangelistsModule::Get().OnReconcileSavedAssets(); })));
+
+			SubSection.AddSeparator("ReconcileClearCacheSeparator");
+
+			SubSection.AddMenuEntry("Clear loaded assets cache", LOCTEXT("SourceControl_ClearLoadedAssetsCache", "Clear loaded assets cache"), LOCTEXT("SourceControl_ClearLoadedAssetsCache_Tooltip", "Clears the cache containing unchecked loaded assets."), FSlateIcon(),
+									FUIAction(FExecuteAction::CreateLambda([]() { FUncontrolledChangelistsModule::Get().OnClearLoadedAssetsCache(); })));
+
+			SubSection.AddMenuEntry("Clear saved assets cache", LOCTEXT("SourceControl_ClearSavedAssetsCache", "Clear saved assets cache"), LOCTEXT("SourceControl_ClearSavedAssetsCache_Tooltip", "Clears the cache containing unchecked saved assets."), FSlateIcon(),
+									FUIAction(FExecuteAction::CreateLambda([]() { FUncontrolledChangelistsModule::Get().OnClearSavedAssetsCache(); })));
+		}));
+	}
+
 	return ToolMenus->GenerateWidget(Menu);
 }
 
@@ -1599,7 +1647,7 @@ void SSourceControlChangelistsWidget::SaveExpandedState(TMap<FSourceControlChang
 {
 	for (FChangelistTreeItemPtr Root : ChangelistsNodes)
 	{
-		if (Root->GetTreeItemType() != IChangelistTreeItem::Changelist)
+		if ((Root->GetTreeItemType() != IChangelistTreeItem::Changelist) && (Root->GetTreeItemType() != IChangelistTreeItem::UncontrolledChangelist))
 		{
 			continue;
 		}
@@ -1628,7 +1676,7 @@ void SSourceControlChangelistsWidget::RestoreExpandedState(const TMap<FSourceCon
 {
 	for (FChangelistTreeItemPtr Root : ChangelistsNodes)
 	{
-		if (Root->GetTreeItemType() != IChangelistTreeItem::Changelist)
+		if ((Root->GetTreeItemType() != IChangelistTreeItem::Changelist) && (Root->GetTreeItemType() != IChangelistTreeItem::UncontrolledChangelist))
 		{
 			continue;
 		}
diff --git a/Engine/Source/Editor/SourceControlWindows/Private/SSourceControlChangelists.h b/Engine/Source/Editor/SourceControlWindows/Private/SSourceControlChangelists.h
--- a/Engine/Source/Editor/SourceControlWindows/Private/SSourceControlChangelists.h
+++ b/Engine/Source/Editor/SourceControlWindows/Private/SSourceControlChangelists.h
@@ -134,6 +134,9 @@ private:
 	/** Source control state changed delegate handle */
 	FDelegateHandle SourceControlStateChangedDelegateHandle;
 
+	/** Uncontrolled Changelist changed delegate handle */
+	FDelegateHandle UncontrolledChangelistChangedDelegateHandle;
+
 	bool bShouldRefresh;
 
 	void StartRefreshStatus();
diff --git a/Engine/Source/Editor/UnrealEd/Private/EditorEngine.cpp b/Engine/Source/Editor/UnrealEd/Private/EditorEngine.cpp
--- a/Engine/Source/Editor/UnrealEd/Private/EditorEngine.cpp
+++ b/Engine/Source/Editor/UnrealEd/Private/EditorEngine.cpp
@@ -1009,6 +1009,7 @@ void UEditorEngine::Init(IEngineLoop* InEngineLoop)
 			TEXT("OutputLog"),
 			TEXT("SourceControl"),
 			TEXT("SourceControlWindows"),
+			TEXT("UncontrolledChangelists"),
 			TEXT("TextureCompressor"),
 			TEXT("MeshUtilities"),
 			TEXT("MovieSceneTools"),
