Adds detection of missing "mark of add" files when reconciling.

#changelist validated

diff --git a/Engine/Source/Developer/UncontrolledChangelists/Private/UncontrolledChangelistState.cpp b/Engine/Source/Developer/UncontrolledChangelists/Private/UncontrolledChangelistState.cpp
--- a/Engine/Source/Developer/UncontrolledChangelists/Private/UncontrolledChangelistState.cpp
+++ b/Engine/Source/Developer/UncontrolledChangelists/Private/UncontrolledChangelistState.cpp
@@ -114,7 +114,7 @@ bool FUncontrolledChangelistState::AddFiles(const TArray<FString>& InFilenames,
 		for (FSourceControlStateRef FileState : FileStates)
 		{
 			bool bIsCheckoutCompliant = (!bCheckCheckout) || (!FileState->IsCheckedOut());
-			bool bIsStatusCompliant = (!bCheckStatus) || FileState->IsModified();
+			bool bIsStatusCompliant = (!bCheckStatus) || FileState->IsModified() || (!FileState->IsSourceControlled());
 
 			if (bIsCheckoutCompliant && bIsStatusCompliant)
 			{
