Updates Uncontrolled Changelists module so unload operations removes files from reconcile cache.
* Removes "clear cache" contextuals actions from reconcile menu.
* Refresh action in Changelist Window now cleans loaded files cache to reconcile

#changelist validated

diff --git a/Engine/Source/Developer/UncontrolledChangelists/Private/UncontrolledChangelistsModule.cpp b/Engine/Source/Developer/UncontrolledChangelists/Private/UncontrolledChangelistsModule.cpp
--- a/Engine/Source/Developer/UncontrolledChangelists/Private/UncontrolledChangelistsModule.cpp
+++ b/Engine/Source/Developer/UncontrolledChangelists/Private/UncontrolledChangelistsModule.cpp
@@ -2,6 +2,7 @@
 
 #include "UncontrolledChangelistsModule.h"
 
+#include "Algo/AnyOf.h"
 #include "Algo/Copy.h"
 #include "Algo/Find.h"
 #include "Algo/Transform.h"
@@ -108,6 +109,8 @@ void FUncontrolledChangelistsModule::UpdateStatus()
 		bHasStateChanged |= UncontrolledChangelistState->UpdateStatus();
 	}
 
+	CleanLoadedAssetsCache();
+
 	if (bHasStateChanged)
 	{
 		OnStateChanged();
@@ -128,7 +131,6 @@ void FUncontrolledChangelistsModule::OnReconcileLoadedAssets()
 
 	CleanAssetsCaches();
 	AddFilesToDefaultUncontrolledChangelist(LoadedFilesCache.Array(), FUncontrolledChangelistState::ECheckFlags::All);
-	LoadedFilesCache.Reset();
 }
 
 void FUncontrolledChangelistsModule::OnReconcileSavedAssets()
@@ -163,16 +165,6 @@ void FUncontrolledChangelistsModule::OnReconcileAllAssets()
 	SavedFilesCache.Reset();
 }
 
-void FUncontrolledChangelistsModule::OnClearLoadedAssetsCache()
-{
-	LoadedFilesCache.Reset();
-}
-
-void FUncontrolledChangelistsModule::OnClearSavedAssetsCache()
-{
-	SavedFilesCache.Reset();
-}
-
 void FUncontrolledChangelistsModule::OnAssetAdded(const struct FAssetData& AssetData)
 {
 
@@ -375,6 +367,40 @@ void FUncontrolledChangelistsModule::CleanAssetsCaches()
 	}
 }
 
+void FUncontrolledChangelistsModule::CleanLoadedAssetsCache()
+{
+	FAssetRegistryModule& AssetRegistryModule = FModuleManager::GetModuleChecked<FAssetRegistryModule>(TEXT("AssetRegistry"));
+	IAssetRegistry& AssetRegistry = AssetRegistryModule.Get();
+
+	for (auto It = LoadedFilesCache.CreateIterator(); It; ++It)
+	{
+		FString& Filename = *It;
+		FString PackageName;
+		TArray<FAssetData> AssetData;
+
+		FPackageName::TryConvertFilenameToLongPackageName(Filename, PackageName);
+
+		if (PackageName.IsEmpty())
+		{
+			continue;
+		}
+
+		AssetRegistry.GetAssetsByPackageName(*PackageName, AssetData);
+
+		if (AssetData.IsEmpty())
+		{
+			continue;
+		}
+
+		bool bIsLoaded = Algo::AnyOf(AssetData, [](const FAssetData& InAssetData) { return InAssetData.IsAssetLoaded(); });
+
+		if (!bIsLoaded)
+		{
+			It.RemoveCurrent();
+		}
+	}
+}
+
 void FUncontrolledChangelistsModule::AddFilesToDefaultUncontrolledChangelist(const TArray<FString>& InFilenames, const FUncontrolledChangelistState::ECheckFlags InCheckFlags)
 {
 	bool bHasStateChanged = false;
diff --git a/Engine/Source/Developer/UncontrolledChangelists/Public/UncontrolledChangelistsModule.h b/Engine/Source/Developer/UncontrolledChangelists/Public/UncontrolledChangelistsModule.h
--- a/Engine/Source/Developer/UncontrolledChangelists/Public/UncontrolledChangelistsModule.h
+++ b/Engine/Source/Developer/UncontrolledChangelists/Public/UncontrolledChangelistsModule.h
@@ -82,12 +82,6 @@ public:
 	 */
 	void OnReconcileAllAssets();
 
-	/** Clears the cache containing unchecked loaded assets. */
-	void OnClearLoadedAssetsCache();
-
-	/** Clears the cache containing unchecked saved assets. */
-	void OnClearSavedAssetsCache();
-
 	/**
 	 * Delegate callback called when assets are added to AssetRegistry.
 	 * @param 	AssetData 	The asset just added.
@@ -166,6 +160,9 @@ private:
 	/** Removes from asset caches files already present in Uncontrolled Changelists */
 	void CleanAssetsCaches();
 
+	/** Removes files not currently loaded from loaded asset cache */
+	void CleanLoadedAssetsCache();
+
 	/**
 	 * Try to add the provided filenames to the default Uncontrolled Changelist.
 	 * @param 	InFilenames 	The files to add.
diff --git a/Engine/Source/Editor/SourceControlWindows/Private/SSourceControlChangelists.cpp b/Engine/Source/Editor/SourceControlWindows/Private/SSourceControlChangelists.cpp
--- a/Engine/Source/Editor/SourceControlWindows/Private/SSourceControlChangelists.cpp
+++ b/Engine/Source/Editor/SourceControlWindows/Private/SSourceControlChangelists.cpp
@@ -1347,13 +1347,6 @@ TSharedPtr<SWidget> SSourceControlChangelistsWidget::OnOpenContextMenu()
 			SubSection.AddMenuEntry("Reconcile saved assets", LOCTEXT("SourceControl_ReconcileSavedAssets", "Reconcile saved assets"), LOCTEXT("SourceControl_ReconcileSavedAssets_Tooltip", "Look for uncontrolled modification in previously saved assets."), FSlateIcon(),
 									FUIAction(FExecuteAction::CreateLambda([]() { FUncontrolledChangelistsModule::Get().OnReconcileSavedAssets(); })));
 
-			SubSection.AddSeparator("ReconcileClearCacheSeparator");
-
-			SubSection.AddMenuEntry("Clear loaded assets cache", LOCTEXT("SourceControl_ClearLoadedAssetsCache", "Clear loaded assets cache"), LOCTEXT("SourceControl_ClearLoadedAssetsCache_Tooltip", "Clears the cache containing unchecked loaded assets."), FSlateIcon(),
-									FUIAction(FExecuteAction::CreateLambda([]() { FUncontrolledChangelistsModule::Get().OnClearLoadedAssetsCache(); })));
-
-			SubSection.AddMenuEntry("Clear saved assets cache", LOCTEXT("SourceControl_ClearSavedAssetsCache", "Clear saved assets cache"), LOCTEXT("SourceControl_ClearSavedAssetsCache_Tooltip", "Clears the cache containing unchecked saved assets."), FSlateIcon(),
-									FUIAction(FExecuteAction::CreateLambda([]() { FUncontrolledChangelistsModule::Get().OnClearSavedAssetsCache(); })));
 		}));
 	}
 
