Use AssetRegistry to add writable files to Reconcile cache.
Adds UncontrolledChangelistValidator, reconciling writable assets, warning the user if new uncontrolled changes are found during changelist validation

#changelist validated

diff --git a/Engine/Plugins/Editor/DataValidation/Source/DataValidation/DataValidation.Build.cs b/Engine/Plugins/Editor/DataValidation/Source/DataValidation/DataValidation.Build.cs
--- a/Engine/Plugins/Editor/DataValidation/Source/DataValidation/DataValidation.Build.cs
+++ b/Engine/Plugins/Editor/DataValidation/Source/DataValidation/DataValidation.Build.cs
@@ -37,7 +37,8 @@ public class DataValidation : ModuleRules
 				"SandboxFile",
 				"Blutility",
 				"ToolMenus",
-				"SourceControl"
+				"SourceControl",
+				"UncontrolledChangelists",
 			}
 		);
 	}
diff --git a/Engine/Plugins/Editor/DataValidation/Source/DataValidation/Private/UncontrolledChangelistValidator.cpp b/Engine/Plugins/Editor/DataValidation/Source/DataValidation/Private/UncontrolledChangelistValidator.cpp
new file mode 100644
--- /dev/null
+++ b/Engine/Plugins/Editor/DataValidation/Source/DataValidation/Private/UncontrolledChangelistValidator.cpp
@@ -0,0 +1,29 @@
+// Copyright Epic Games, Inc. All Rights Reserved.
+
+#include "UncontrolledChangelistValidator.h"
+
+#include "DataValidationChangelist.h"
+#include "UncontrolledChangelistsModule.h"
+
+#define LOCTEXT_NAMESPACE "UncontrolledChangelistValidation"
+
+bool UUncontrolledChangelistValidator::CanValidateAsset_Implementation(UObject* InAsset) const
+{
+	return (InAsset != nullptr) && (UDataValidationChangelist::StaticClass() == InAsset->GetClass());
+}
+
+EDataValidationResult UUncontrolledChangelistValidator::ValidateLoadedAsset_Implementation(UObject* InAsset, TArray<FText>& ValidationErrors)
+{
+	FUncontrolledChangelistsModule& UncontrolledChangelistModule = FUncontrolledChangelistsModule::Get();
+
+	if (UncontrolledChangelistModule.OnReconcileAssets())
+	{
+		AssetFails(InAsset, LOCTEXT("UncontrolledChangesFound", "Uncontrolled changes found, please verify if they should be added to your changelist."), ValidationErrors);
+		return EDataValidationResult::NotValidated;
+	}
+
+	AssetPasses(InAsset);
+	return EDataValidationResult::Valid;
+}
+
+#undef LOCTEXT_NAMESPACE
diff --git a/Engine/Plugins/Editor/DataValidation/Source/DataValidation/Public/UncontrolledChangelistValidator.h b/Engine/Plugins/Editor/DataValidation/Source/DataValidation/Public/UncontrolledChangelistValidator.h
new file mode 100644
--- /dev/null
+++ b/Engine/Plugins/Editor/DataValidation/Source/DataValidation/Public/UncontrolledChangelistValidator.h
@@ -0,0 +1,19 @@
+// Copyright Epic Games, Inc. All Rights Reserved.
+
+#pragma once
+
+#include "EditorValidatorBase.h"
+#include "UncontrolledChangelistValidator.generated.h"
+
+/**
+* Validates there is no additional uncontrolled changes waiting to be reconciled.
+*/
+UCLASS()
+class DATAVALIDATION_API UUncontrolledChangelistValidator : public UEditorValidatorBase
+{
+	GENERATED_BODY()
+
+protected:
+	virtual bool CanValidateAsset_Implementation(UObject* InAsset) const override;
+	virtual EDataValidationResult ValidateLoadedAsset_Implementation(UObject* InAsset, TArray<FText>& ValidationErrors) override;
+};
diff --git a/Engine/Source/Developer/UncontrolledChangelists/Private/UncontrolledChangelistState.cpp b/Engine/Source/Developer/UncontrolledChangelists/Private/UncontrolledChangelistState.cpp
--- a/Engine/Source/Developer/UncontrolledChangelists/Private/UncontrolledChangelistState.cpp
+++ b/Engine/Source/Developer/UncontrolledChangelists/Private/UncontrolledChangelistState.cpp
@@ -179,14 +179,13 @@ bool FUncontrolledChangelistState::UpdateStatus()
 	return bOutChanged;
 }
 
-void FUncontrolledChangelistState::RemoveDuplicates(TSet<FString>& InOutLoadedFiles, TSet<FString>& InOutModifiedFiles)
+void FUncontrolledChangelistState::RemoveDuplicates(TSet<FString>& InOutAddedAssets)
 {
 	for (const FSourceControlStateRef& FileState : Files)
 	{
 		const FString& Filename = FileState->GetFilename();
 		
-		InOutLoadedFiles.Remove(Filename);
-		InOutModifiedFiles.Remove(Filename);
+		InOutAddedAssets.Remove(Filename);
 	}
 }
 
diff --git a/Engine/Source/Developer/UncontrolledChangelists/Private/UncontrolledChangelistsModule.cpp b/Engine/Source/Developer/UncontrolledChangelists/Private/UncontrolledChangelistsModule.cpp
--- a/Engine/Source/Developer/UncontrolledChangelists/Private/UncontrolledChangelistsModule.cpp
+++ b/Engine/Source/Developer/UncontrolledChangelists/Private/UncontrolledChangelistsModule.cpp
@@ -5,6 +5,7 @@
 #include "Algo/AnyOf.h"
 #include "Algo/Copy.h"
 #include "Algo/Find.h"
+#include "Algo/ForEach.h"
 #include "Algo/Transform.h"
 #include "AssetRegistryModule.h"
 #include "Dom/JsonObject.h"
@@ -38,9 +39,14 @@ void FUncontrolledChangelistsModule::StartupModule()
 	LoadState();
 
 	FAssetRegistryModule& AssetRegistryModule = FModuleManager::GetModuleChecked<FAssetRegistryModule>(TEXT("AssetRegistry"));
-	OnAssetAddedDelegateHandle = AssetRegistryModule.Get().OnAssetAdded().AddLambda([](const struct FAssetData& AssetData) { Get().OnAssetAdded(AssetData); });
-	OnAssetLoadedDelegateHandle = FCoreUObjectDelegates::OnAssetLoaded.AddLambda([](UObject* InAsset) { Get().OnAssetLoaded(InAsset); });
+	IAssetRegistry& AssetRegistry = AssetRegistryModule.Get();
+	OnAssetAddedDelegateHandle = AssetRegistry.OnAssetAdded().AddLambda([](const struct FAssetData& AssetData) { Get().OnAssetAdded(AssetData); });
 	OnObjectPreSavedDelegateHandle = FCoreUObjectDelegates::OnObjectPreSave.AddLambda([](UObject* InAsset, const FObjectPreSaveContext& InPreSaveContext) { Get().OnObjectPreSaved(InAsset, InPreSaveContext); });
+
+	TArray<FAssetData> AssetData;
+
+	AssetRegistry.GetAllAssets(AssetData);
+	Algo::ForEach(AssetData, [this](const struct FAssetData& AssetData) { OnAssetAdded(AssetData); });
 }
 
 void FUncontrolledChangelistsModule::ShutdownModule()
@@ -53,7 +59,6 @@ void FUncontrolledChangelistsModule::ShutdownModule()
 		AssetRegistryModulePtr->Get().OnAssetAdded().Remove(OnAssetAddedDelegateHandle);
 	}
 
-	FCoreUObjectDelegates::OnAssetLoaded.Remove(OnAssetLoadedDelegateHandle);
 	FCoreUObjectDelegates::OnObjectPreSave.Remove(OnObjectPreSavedDelegateHandle);
 }
 
@@ -109,8 +114,6 @@ void FUncontrolledChangelistsModule::UpdateStatus()
 		bHasStateChanged |= UncontrolledChangelistState->UpdateStatus();
 	}
 
-	CleanLoadedAssetsCache();
-
 	if (bHasStateChanged)
 	{
 		OnStateChanged();
@@ -119,72 +122,56 @@ void FUncontrolledChangelistsModule::UpdateStatus()
 
 FText FUncontrolledChangelistsModule::GetReconcileStatus() const
 {
-	return FText::Format(LOCTEXT("ReconcileStatus", "Assets to check for reconcile: {0} loaded, {1} saved"), FText::AsNumber(LoadedFilesCache.Num()), FText::AsNumber(SavedFilesCache.Num()));
+	return FText::Format(LOCTEXT("ReconcileStatus", "Assets to check for reconcile: {0}"), FText::AsNumber(AddedAssetsCache.Num()));
 }
 
-void FUncontrolledChangelistsModule::OnReconcileLoadedAssets()
+bool FUncontrolledChangelistsModule::OnReconcileAssets()
 {
-	if ((!IsEnabled()) || ((LoadedFilesCache.Num() == 0)))
+	if ((!IsEnabled()) || AddedAssetsCache.IsEmpty())
 	{
-		return;
+		return false;
 	}
 
 	CleanAssetsCaches();
-	AddFilesToDefaultUncontrolledChangelist(LoadedFilesCache.Array(), FUncontrolledChangelistState::ECheckFlags::All);
+	bool bHasStateChanged = AddFilesToDefaultUncontrolledChangelist(AddedAssetsCache.Array(), FUncontrolledChangelistState::ECheckFlags::All);
+
+	AddedAssetsCache.Empty();
+
+	return bHasStateChanged;
 }
 
-void FUncontrolledChangelistsModule::OnReconcileSavedAssets()
+void FUncontrolledChangelistsModule::OnAssetAdded(const struct FAssetData& AssetData)
 {
-	if ((!IsEnabled()) || ((SavedFilesCache.Num() == 0)))
+	if (!IsEnabled())
 	{
 		return;
 	}
 
-	CleanAssetsCaches();
-	AddFilesToDefaultUncontrolledChangelist(SavedFilesCache.Array(), FUncontrolledChangelistState::ECheckFlags::All);
-	SavedFilesCache.Reset();
-}
+	FPackagePath PackagePath;
 
-void FUncontrolledChangelistsModule::OnReconcileAllAssets()
-{
-	if ((!IsEnabled()) || ((LoadedFilesCache.Num() == 0) && (SavedFilesCache.Num() == 0)))
+	if (!FPackagePath::TryFromPackageName(AssetData.PackageName, PackagePath))
 	{
 		return;
 	}
 
-	CleanAssetsCaches();
-
-	TArray<FString> FilenamesToCheck;
-
-	Algo::Copy(LoadedFilesCache, FilenamesToCheck);
-	Algo::Copy(SavedFilesCache, FilenamesToCheck);
+	const FString LocalFullPath(PackagePath.GetLocalFullPath());
 
-	AddFilesToDefaultUncontrolledChangelist(FilenamesToCheck, FUncontrolledChangelistState::ECheckFlags::All);
-
-	LoadedFilesCache.Reset();
-	SavedFilesCache.Reset();
-}
-
-void FUncontrolledChangelistsModule::OnAssetAdded(const struct FAssetData& AssetData)
-{
-
-}
-
-void FUncontrolledChangelistsModule::OnAssetLoaded(UObject* InAsset)
-{
-	if (!IsEnabled())
+	if (LocalFullPath.IsEmpty())
 	{
 		return;
 	}
 
-	FString Fullpath = GetUObjectPackageFullpath(InAsset);
+	FString Fullpath = FPaths::ConvertRelativePathToFull(LocalFullPath);
 
 	if (Fullpath.IsEmpty())
 	{
 		return;
 	}
 
-	LoadedFilesCache.Add(MoveTemp(Fullpath));
+	if (!IFileManager::Get().IsReadOnly(*Fullpath))
+	{
+		AddedAssetsCache.Add(MoveTemp(Fullpath));
+	}
 }
 
 void FUncontrolledChangelistsModule::OnObjectPreSaved(UObject* InAsset, const FObjectPreSaveContext& InPreSaveContext)
@@ -201,7 +188,7 @@ void FUncontrolledChangelistsModule::OnObjectPreSaved(UObject* InAsset, const FO
 		return;
 	}
 
-	SavedFilesCache.Add(MoveTemp(Fullpath));
+	AddedAssetsCache.Add(MoveTemp(Fullpath));
 }
 
 void FUncontrolledChangelistsModule::MoveFilesToUncontrolledChangelist(const TArray<FSourceControlStateRef>& InControlledFileStates, const TArray<FSourceControlStateRef>& InUncontrolledFileStates, const FUncontrolledChangelist& InUncontrolledChangelist)
@@ -363,45 +350,11 @@ void FUncontrolledChangelistsModule::CleanAssetsCaches()
 	for (FUncontrolledChangelistsStateCache::ElementType& Pair : UncontrolledChangelistsStateCache)
 	{
 		FUncontrolledChangelistsStateCache::ValueType& UncontrolledChangelistState = Pair.Value;
-		UncontrolledChangelistState->RemoveDuplicates(LoadedFilesCache, SavedFilesCache);
+		UncontrolledChangelistState->RemoveDuplicates(AddedAssetsCache);
 	}
 }
 
-void FUncontrolledChangelistsModule::CleanLoadedAssetsCache()
-{
-	FAssetRegistryModule& AssetRegistryModule = FModuleManager::GetModuleChecked<FAssetRegistryModule>(TEXT("AssetRegistry"));
-	IAssetRegistry& AssetRegistry = AssetRegistryModule.Get();
-
-	for (auto It = LoadedFilesCache.CreateIterator(); It; ++It)
-	{
-		FString& Filename = *It;
-		FString PackageName;
-		TArray<FAssetData> AssetData;
-
-		FPackageName::TryConvertFilenameToLongPackageName(Filename, PackageName);
-
-		if (PackageName.IsEmpty())
-		{
-			continue;
-		}
-
-		AssetRegistry.GetAssetsByPackageName(*PackageName, AssetData);
-
-		if (AssetData.IsEmpty())
-		{
-			continue;
-		}
-
-		bool bIsLoaded = Algo::AnyOf(AssetData, [](const FAssetData& InAssetData) { return InAssetData.IsAssetLoaded(); });
-
-		if (!bIsLoaded)
-		{
-			It.RemoveCurrent();
-		}
-	}
-}
-
-void FUncontrolledChangelistsModule::AddFilesToDefaultUncontrolledChangelist(const TArray<FString>& InFilenames, const FUncontrolledChangelistState::ECheckFlags InCheckFlags)
+bool FUncontrolledChangelistsModule::AddFilesToDefaultUncontrolledChangelist(const TArray<FString>& InFilenames, const FUncontrolledChangelistState::ECheckFlags InCheckFlags)
 {
 	bool bHasStateChanged = false;
 
@@ -415,6 +368,8 @@ void FUncontrolledChangelistsModule::AddFilesToDefaultUncontrolledChangelist(con
 	{
 		OnStateChanged();
 	}
+
+	return bHasStateChanged;
 }
 
 void FUncontrolledChangelistsModule::SaveState() const
diff --git a/Engine/Source/Developer/UncontrolledChangelists/Public/UncontrolledChangelistState.h b/Engine/Source/Developer/UncontrolledChangelists/Public/UncontrolledChangelistState.h
--- a/Engine/Source/Developer/UncontrolledChangelists/Public/UncontrolledChangelistState.h
+++ b/Engine/Source/Developer/UncontrolledChangelists/Public/UncontrolledChangelistState.h
@@ -106,11 +106,10 @@ public:
 	bool UpdateStatus();
 
 	/**
-	 * Removes files present both in the Uncontrolled Changelist and the provided sets.
-	 * @param 	InOutLoadedFiles 	A Set representing loaded files to check against.
-	 * @param 	InOutModifiedFiles 	A Set representing modified files to check against.
+	 * Removes files present both in the Uncontrolled Changelist and the provided set.
+	 * @param 	InOutAddedAssets 	A Set representing Added Assets to check against.
 	 */
-	void RemoveDuplicates(TSet<FString>& InOutLoadedFiles, TSet<FString>& InOutModifiedFiles);
+	void RemoveDuplicates(TSet<FString>& InOutAddedAssets);
 
 public:
 	FUncontrolledChangelist Changelist;
diff --git a/Engine/Source/Developer/UncontrolledChangelists/Public/UncontrolledChangelistsModule.h b/Engine/Source/Developer/UncontrolledChangelists/Public/UncontrolledChangelistsModule.h
--- a/Engine/Source/Developer/UncontrolledChangelists/Public/UncontrolledChangelistsModule.h
+++ b/Engine/Source/Developer/UncontrolledChangelists/Public/UncontrolledChangelistsModule.h
@@ -67,20 +67,11 @@ public:
 	 */
 	FText GetReconcileStatus() const;
 
-	/** Called when "Reconcile loaded assets" button is clicked. Checks for uncontrolled modifications in previously loaded assets.
+	/** Called when "Reconcile assets" button is clicked. Checks for uncontrolled modifications in previously added assets.
 	 *	Adds modified files to Uncontrolled Changelists
+	 *  @return True if new modifications found
 	 */
-	void OnReconcileLoadedAssets();
-
-	/** Called when "Reconcile saved assets" button is clicked. Checks for uncontrolled modifications in previously saved assets.
-	 *	Adds modified files to Uncontrolled Changelists
-	 */
-	void OnReconcileSavedAssets();
-
-	/** Called when the "Reconcile all assets" button is clicked. Checks for uncontrolled modifications in previously loaded and saved assets.
-	 *	Adds modified files to Uncontrolled Changelists
-	 */
-	void OnReconcileAllAssets();
+	bool OnReconcileAssets();
 
 	/**
 	 * Delegate callback called when assets are added to AssetRegistry.
@@ -160,21 +151,17 @@ private:
 	/** Removes from asset caches files already present in Uncontrolled Changelists */
 	void CleanAssetsCaches();
 
-	/** Removes files not currently loaded from loaded asset cache */
-	void CleanLoadedAssetsCache();
-
 	/**
 	 * Try to add the provided filenames to the default Uncontrolled Changelist.
 	 * @param 	InFilenames 	The files to add.
 	 * @param 	InCheckFlags 	The required checks to check the file against before adding.
+	 * @return 	True file have been added.
 	 */
-	void AddFilesToDefaultUncontrolledChangelist(const TArray<FString>& InFilenames, const FUncontrolledChangelistState::ECheckFlags InCheckFlags);
+	bool AddFilesToDefaultUncontrolledChangelist(const TArray<FString>& InFilenames, const FUncontrolledChangelistState::ECheckFlags InCheckFlags);
 
 private:
 	FUncontrolledChangelistsStateCache	UncontrolledChangelistsStateCache;
-	TSet<FString>						LoadedFilesCache;
-	TSet<FString>						SavedFilesCache;
+	TSet<FString>						AddedAssetsCache;
 	FDelegateHandle						OnAssetAddedDelegateHandle;
-	FDelegateHandle						OnAssetLoadedDelegateHandle;
 	FDelegateHandle						OnObjectPreSavedDelegateHandle;
 };
diff --git a/Engine/Source/Editor/SourceControlWindows/Private/SSourceControlChangelists.cpp b/Engine/Source/Editor/SourceControlWindows/Private/SSourceControlChangelists.cpp
--- a/Engine/Source/Editor/SourceControlWindows/Private/SSourceControlChangelists.cpp
+++ b/Engine/Source/Editor/SourceControlWindows/Private/SSourceControlChangelists.cpp
@@ -1334,20 +1334,8 @@ TSharedPtr<SWidget> SSourceControlChangelistsWidget::OnOpenContextMenu()
 	{
 		Section.AddSeparator("ReconcileSeparator");
 
-		Section.AddSubMenu("ReconcileSubMenu", LOCTEXT("ReconcileSubMenu", "Reconcile"), LOCTEXT("ReconcileSubMenu_ToolTip", "Opens the Reconcile sub menu"), FNewToolMenuDelegate::CreateLambda([](UToolMenu* Menu)
-		{
-			FToolMenuSection& SubSection = Menu->AddSection(TEXT("Reconcile"));
-			
-			SubSection.AddMenuEntry("Reconcile all assets", LOCTEXT("SourceControl_ReconcileAllAssets", "Reconcile all assets"), LOCTEXT("SourceControl_ReconcileAllAssets_Tooltip", "Look for uncontrolled modification in previously loaded and saved assets."), FSlateIcon(),
-									FUIAction(FExecuteAction::CreateLambda([]() { FUncontrolledChangelistsModule::Get().OnReconcileAllAssets(); })));
-
-			SubSection.AddMenuEntry("Reconcile loaded assets", LOCTEXT("SourceControl_ReconcileLoadedAssets", "Reconcile loaded assets"), LOCTEXT("SourceControl_ReconcileLoadedAssets_Tooltip", "Look for uncontrolled modification in previously loaded assets."), FSlateIcon(),
-									FUIAction(FExecuteAction::CreateLambda([]() { FUncontrolledChangelistsModule::Get().OnReconcileLoadedAssets(); })));
-
-			SubSection.AddMenuEntry("Reconcile saved assets", LOCTEXT("SourceControl_ReconcileSavedAssets", "Reconcile saved assets"), LOCTEXT("SourceControl_ReconcileSavedAssets_Tooltip", "Look for uncontrolled modification in previously saved assets."), FSlateIcon(),
-									FUIAction(FExecuteAction::CreateLambda([]() { FUncontrolledChangelistsModule::Get().OnReconcileSavedAssets(); })));
-
-		}));
+		Section.AddMenuEntry("Reconcile assets", LOCTEXT("SourceControl_ReconcileAssets", "Reconcile assets"), LOCTEXT("SourceControl_ReconcileAssets_Tooltip", "Look for uncontrolled modification in currently added assets."), FSlateIcon(),
+								FUIAction(FExecuteAction::CreateLambda([]() { FUncontrolledChangelistsModule::Get().OnReconcileAssets(); })));
 	}
 
 	return ToolMenus->GenerateWidget(Menu);
