Fixes files not in depot and deleted staying in Uncontrolled Changelist persistency file.

#changelist validated

diff --git a/Engine/Source/Developer/UncontrolledChangelists/Private/UncontrolledChangelistState.cpp b/Engine/Source/Developer/UncontrolledChangelists/Private/UncontrolledChangelistState.cpp
--- a/Engine/Source/Developer/UncontrolledChangelists/Private/UncontrolledChangelistState.cpp
+++ b/Engine/Source/Developer/UncontrolledChangelists/Private/UncontrolledChangelistState.cpp
@@ -6,6 +6,7 @@
 #include "Algo/Transform.h"
 #include "Dom/JsonObject.h"
 #include "Dom/JsonValue.h"
+#include "HAL/FileManager.h"
 #include "ISourceControlModule.h"
 #include "ISourceControlProvider.h"
 #include "SourceControlHelpers.h"
@@ -119,14 +120,20 @@ bool FUncontrolledChangelistState::AddFiles(const TArray<FString>& InFilenames,
 		SourceControlProvider.Execute(UpdateStatusOperation, InFilenames);
 	}
 
-	bool GetStateSucceeded = SourceControlProvider.GetState(InFilenames, FileStates, EStateCacheUsage::Use) == ECommandResult::Succeeded;
+	const bool GetStateSucceeded = SourceControlProvider.GetState(InFilenames, FileStates, EStateCacheUsage::Use) == ECommandResult::Succeeded;
 
 	if (GetStateSucceeded && (!FileStates.IsEmpty()))
 	{
 		for (FSourceControlStateRef FileState : FileStates)
 		{
-			bool bIsCheckoutCompliant = (!bCheckCheckout) || (!FileState->IsCheckedOut());
-			bool bIsStatusCompliant = (!bCheckStatus) || FileState->IsModified() || (!FileState->IsSourceControlled());
+			const bool bIsSourceControlled = FileState->IsSourceControlled();
+			const bool bFileExists = IFileManager::Get().FileExists(*FileState->GetFilename());
+
+			const bool bIsUncontrolled = (!bIsSourceControlled) && bFileExists;
+			const bool bIsDeleted = FileState->IsSourceControlled() && (!bFileExists);
+
+			const bool bIsCheckoutCompliant = (!bCheckCheckout) || (!FileState->IsCheckedOut());
+			const bool bIsStatusCompliant = (!bCheckStatus) || FileState->IsModified() || bIsUncontrolled || bIsDeleted;
 
 			if (bIsCheckoutCompliant && bIsStatusCompliant)
 			{
