Removes abusive Uncontrolled Changelist warning. Replaced by display log when loading performed as expected.

#changelist validated

diff --git a/Engine/Source/Developer/UncontrolledChangelists/Private/UncontrolledChangelistsModule.cpp b/Engine/Source/Developer/UncontrolledChangelists/Private/UncontrolledChangelistsModule.cpp
--- a/Engine/Source/Developer/UncontrolledChangelists/Private/UncontrolledChangelistsModule.cpp
+++ b/Engine/Source/Developer/UncontrolledChangelists/Private/UncontrolledChangelistsModule.cpp
@@ -412,7 +412,6 @@ void FUncontrolledChangelistsModule::LoadState()
 
 	if (!FFileHelper::LoadFileToString(ImportJsonString, *GetPersistentFilePath()))
 	{
-		UE_LOG(LogSourceControl, Warning, TEXT("Cannot load file %s."), *GetPersistentFilePath());
 		return;
 	}
 
@@ -457,6 +456,8 @@ void FUncontrolledChangelistsModule::LoadState()
 
 		UncontrolledChangelistState->Deserialize(JsonObject);
 	}
+
+	UE_LOG(LogSourceControl, Display, TEXT("Uncontrolled Changelist persistency file loaded %s"), *GetPersistentFilePath());
 }
 
 FString FUncontrolledChangelistsModule::GetPersistentFilePath() const
