Prevent saved file to be added to PackagesNotToPromptAnyMore when added to UncontrolledChangelist.
Allows save window to be displayed again if a file is reverted in UncontrolledChangelist.

#changelist validated

diff --git a/Engine/Source/Developer/UncontrolledChangelists/Private/UncontrolledChangelistsModule.cpp b/Engine/Source/Developer/UncontrolledChangelists/Private/UncontrolledChangelistsModule.cpp
--- a/Engine/Source/Developer/UncontrolledChangelists/Private/UncontrolledChangelistsModule.cpp
+++ b/Engine/Source/Developer/UncontrolledChangelists/Private/UncontrolledChangelistsModule.cpp
@@ -79,13 +79,13 @@ TArray<FUncontrolledChangelistStateRef> FUncontrolledChangelistsModule::GetChang
 	return UncontrolledChangelistStates;
 }
 
-void FUncontrolledChangelistsModule::OnMakeWritable(const FString& InFilename)
+bool FUncontrolledChangelistsModule::OnMakeWritable(const FString& InFilename)
 {
 	bool bHasStateChanged = false;
 
 	if (!IsEnabled())
 	{
-		return;
+		return false;
 	}
 
 	FUncontrolledChangelist DefaultUncontrolledChangelist(FUncontrolledChangelist::DEFAULT_UNCONTROLLED_CHANGELIST_GUID, FUncontrolledChangelist::DEFAULT_UNCONTROLLED_CHANGELIST_NAME.ToString());
@@ -97,6 +97,8 @@ void FUncontrolledChangelistsModule::OnMakeWritable(const FString& InFilename)
 	{
 		OnStateChanged();
 	}
+
+	return bHasStateChanged;
 }
 
 void FUncontrolledChangelistsModule::UpdateStatus()
diff --git a/Engine/Source/Developer/UncontrolledChangelists/Public/UncontrolledChangelistsModule.h b/Engine/Source/Developer/UncontrolledChangelists/Public/UncontrolledChangelistsModule.h
--- a/Engine/Source/Developer/UncontrolledChangelists/Public/UncontrolledChangelistsModule.h
+++ b/Engine/Source/Developer/UncontrolledChangelists/Public/UncontrolledChangelistsModule.h
@@ -43,8 +43,9 @@ public:
 	 /**
 	  * Called when a file has been made writable. Adds the file to the Default Uncontrolled Changelist
 	  * @param	InFilename			The file to be added.
+	  * @return True if the file has been added to the Default Uncontrolled Changelist.
 	  */
-	 void OnMakeWritable(const FString& InFilename);
+	 bool OnMakeWritable(const FString& InFilename);
 
 	/**
 	 * Updates the status of Uncontrolled Changelists and files.
diff --git a/Engine/Source/Editor/UnrealEd/Private/FileHelpers.cpp b/Engine/Source/Editor/UnrealEd/Private/FileHelpers.cpp
--- a/Engine/Source/Editor/UnrealEd/Private/FileHelpers.cpp
+++ b/Engine/Source/Editor/UnrealEd/Private/FileHelpers.cpp
@@ -1834,9 +1834,13 @@ bool FEditorFileUtils::PromptToCheckoutPackages(bool bCheckDirty, const TArray<U
 						// Knock off the read only flag from the current file attributes
 						if (FPlatformFileManager::Get().GetPlatformFile().SetReadOnly(*Filename, false))
 						{
-							UncontrolledChangelistModule.OnMakeWritable(Filename);
+							// Add to PackagesNotToPromptAnyMore only if not added to Uncontrolled Changelist.
+							// If added to Uncontrolled Changelist, we want the checkout prompt to be displayed again if the file is reverted
+							if (!UncontrolledChangelistModule.OnMakeWritable(Filename))
+							{
+								PackagesNotToPromptAnyMore.Add(PackageToMakeWritable->GetName());
+							}
 
-							PackagesNotToPromptAnyMore.Add(PackageToMakeWritable->GetName());
 							if (OutPackagesCheckedOutOrMadeWritable)
 							{
 								OutPackagesCheckedOutOrMadeWritable->Add(PackageToMakeWritable);
