UnrealEd: Fix warning complaining about GetLocalFullPath being called on FPackagePath with an unspecified header extension for a path that does not exist on disk. Partially addresses jira UE-137902

Warnings were generated when converting brushes to static meshes, creating geometry collections from static meshes in fracture mode and one instance also occured while loading assets on QAGame startup

#rnx

diff --git a/Engine/Source/Developer/UncontrolledChangelists/Private/UncontrolledChangelistsModule.cpp b/Engine/Source/Developer/UncontrolledChangelists/Private/UncontrolledChangelistsModule.cpp
--- a/Engine/Source/Developer/UncontrolledChangelists/Private/UncontrolledChangelistsModule.cpp
+++ b/Engine/Source/Developer/UncontrolledChangelists/Private/UncontrolledChangelistsModule.cpp
@@ -155,6 +155,11 @@ void FUncontrolledChangelistsModule::OnAssetAdded(const struct FAssetData& Asset
 	{
 		return;
 	}
+	
+	if (!FPackageName::DoesPackageExist(PackagePath, &PackagePath))
+	{
+		return; // If the package does not exist on disk there is nothing more to do
+	}
 
 	const FString LocalFullPath(PackagePath.GetLocalFullPath());
 
