Adds mark for delete files support for Uncontrolled Changelist

#changelist validated

diff --git a/Engine/Source/Developer/UncontrolledChangelists/Private/UncontrolledChangelistsModule.cpp b/Engine/Source/Developer/UncontrolledChangelists/Private/UncontrolledChangelistsModule.cpp
--- a/Engine/Source/Developer/UncontrolledChangelists/Private/UncontrolledChangelistsModule.cpp
+++ b/Engine/Source/Developer/UncontrolledChangelists/Private/UncontrolledChangelistsModule.cpp
@@ -10,6 +10,7 @@
 #include "AssetRegistryModule.h"
 #include "Dom/JsonObject.h"
 #include "FileHelpers.h"
+#include "HAL/FileManager.h"
 #include "HAL/IConsoleManager.h"
 #include "ISourceControlModule.h"
 #include "ISourceControlProvider.h"
@@ -271,6 +272,7 @@ void FUncontrolledChangelistsModule::MoveFilesToControlledChangelist(const TArra
 	TArray<UPackage*> PackageConflicts;
 	TArray<FString> FilesToAdd;
 	TArray<FString> FilesToCheckout;
+	TArray<FString> FilesToDelete;
 
 	// Check if we can Checkout files or mark for add
 	for (const FSourceControlStateRef& Filestate : UpdatedFilestates)
@@ -279,7 +281,11 @@ void FUncontrolledChangelistsModule::MoveFilesToControlledChangelist(const TArra
 		{
 			FilesToAdd.Add(Filestate->GetFilename());
 		}
-		else if(Filestate->CanCheckout())
+		else if (!IFileManager::Get().FileExists(*Filestate->GetFilename()))
+		{
+			FilesToDelete.Add(Filestate->GetFilename());
+		}
+		else if (Filestate->CanCheckout())
 		{
 			FilesToCheckout.Add(Filestate->GetFilename());
 		}
@@ -305,6 +311,11 @@ void FUncontrolledChangelistsModule::MoveFilesToControlledChangelist(const TArra
 			SourceControlProvider.Execute(ISourceControlOperation::Create<FMarkForAdd>(), InChangelist, FilesToAdd);
 		}
 
+		if (!FilesToDelete.IsEmpty())
+		{
+			SourceControlProvider.Execute(ISourceControlOperation::Create<FDelete>(), InChangelist, FilesToDelete);
+		}
+
 		// UpdateStatus so UncontrolledChangelists can remove files from their cache if they were present before checkout.
 		UpdateStatus();
 	}
