#UE Add an IsEnabled check before doing very expensive GetAssets call during startup in UncontrolledChangelistsModule

diff --git a/Engine/Source/Developer/UncontrolledChangelists/Private/UncontrolledChangelistsModule.cpp b/Engine/Source/Developer/UncontrolledChangelists/Private/UncontrolledChangelistsModule.cpp
--- a/Engine/Source/Developer/UncontrolledChangelists/Private/UncontrolledChangelistsModule.cpp
+++ b/Engine/Source/Developer/UncontrolledChangelists/Private/UncontrolledChangelistsModule.cpp
@@ -33,6 +33,11 @@ static TAutoConsoleVariable<bool> CVarUncontrolledChangelistsEnable(
 
 void FUncontrolledChangelistsModule::StartupModule()
 {
+	if (!IsEnabled())
+	{
+		return;
+	}
+
 	// Adds Default Uncontrolled Changelist if it is not already present.
 	FUncontrolledChangelist DefaultUncontrolledChangelist(FUncontrolledChangelist::DEFAULT_UNCONTROLLED_CHANGELIST_GUID, FUncontrolledChangelist::DEFAULT_UNCONTROLLED_CHANGELIST_NAME.ToString());
 	UncontrolledChangelistsStateCache.FindOrAdd(MoveTemp(DefaultUncontrolledChangelist), MakeShareable(new FUncontrolledChangelistState(DefaultUncontrolledChangelist)));
