Uncontrolled Changelist
-Startup scan is now a background task (Saves 1min15 editor startup time on projects with lots of assets)
-Wait on startup task with dialog on reconcile if it isn't finished
-Don't allow changing module enabled state past startup (CVar read at startup only)

diff --git a/Engine/Source/Developer/UncontrolledChangelists/Private/UncontrolledChangelistsModule.cpp b/Engine/Source/Developer/UncontrolledChangelists/Private/UncontrolledChangelistsModule.cpp
--- a/Engine/Source/Developer/UncontrolledChangelists/Private/UncontrolledChangelistsModule.cpp
+++ b/Engine/Source/Developer/UncontrolledChangelists/Private/UncontrolledChangelistsModule.cpp
@@ -16,6 +16,7 @@
 #include "ISourceControlProvider.h"
 #include "Misc/FileHelper.h"
 #include "Misc/Paths.h"
+#include "Misc/ScopedSlowTask.h"
 #include "PackagesDialog.h"
 #include "Serialization/JsonSerializer.h"
 #include "Serialization/JsonWriter.h"
@@ -31,8 +32,28 @@ static TAutoConsoleVariable<bool> CVarUncontrolledChangelistsEnable(
 	TEXT("Enables Uncontrolled Changelists (experimental).")
 );
 
+void FUncontrolledChangelistsModule::FStartupTask::DoWork()
+{
+	double StartTime = FPlatformTime::Seconds();
+	UE_LOG(LogSourceControl, Log, TEXT("Uncontrolled asset enumeration started..."));
+
+	FAssetRegistryModule& AssetRegistryModule = FModuleManager::GetModuleChecked<FAssetRegistryModule>(TEXT("AssetRegistry"));
+	IAssetRegistry& AssetRegistry = AssetRegistryModule.Get();
+
+	TArray<FAssetData> Assets;
+	const bool bIncludeOnlyOnDiskAssets = true;
+	AssetRegistry.GetAllAssets(Assets, bIncludeOnlyOnDiskAssets);
+	Algo::ForEach(Assets, [this](const struct FAssetData& AssetData) 
+	{ 
+		Owner->OnAssetAddedInternal(AssetData, AddedAssetsCache, true);
+	});
+	
+	UE_LOG(LogSourceControl, Log, TEXT("Uncontrolled asset enumeration finished in %s seconds (Found %d uncontrolled assets)"), *FString::SanitizeFloat(FPlatformTime::Seconds() - StartTime), AddedAssetsCache.Num());
+}
+
 void FUncontrolledChangelistsModule::StartupModule()
 {
+	bIsEnabled = CVarUncontrolledChangelistsEnable.GetValueOnGameThread();
 	if (!IsEnabled())
 	{
 		return;
@@ -48,11 +69,9 @@ void FUncontrolledChangelistsModule::StartupModule()
 	IAssetRegistry& AssetRegistry = AssetRegistryModule.Get();
 	OnAssetAddedDelegateHandle = AssetRegistry.OnAssetAdded().AddLambda([](const struct FAssetData& AssetData) { Get().OnAssetAdded(AssetData); });
 	OnObjectPreSavedDelegateHandle = FCoreUObjectDelegates::OnObjectPreSave.AddLambda([](UObject* InAsset, const FObjectPreSaveContext& InPreSaveContext) { Get().OnObjectPreSaved(InAsset, InPreSaveContext); });
-
-	TArray<FAssetData> AssetData;
-
-	AssetRegistry.GetAllAssets(AssetData);
-	Algo::ForEach(AssetData, [this](const struct FAssetData& AssetData) { OnAssetAdded(AssetData); });
+		
+	StartupTask = MakeUnique<FAsyncTask<FStartupTask>>(this);		
+	StartupTask->StartBackgroundTask();
 }
 
 void FUncontrolledChangelistsModule::ShutdownModule()
@@ -70,7 +89,7 @@ void FUncontrolledChangelistsModule::ShutdownModule()
 
 bool FUncontrolledChangelistsModule::IsEnabled() const
 {
-	return CVarUncontrolledChangelistsEnable.GetValueOnGameThread();
+	return bIsEnabled;
 }
 
 TArray<FUncontrolledChangelistStateRef> FUncontrolledChangelistsModule::GetChangelistStates() const
@@ -130,16 +149,39 @@ void FUncontrolledChangelistsModule::UpdateStatus()
 
 FText FUncontrolledChangelistsModule::GetReconcileStatus() const
 {
+	if (StartupTask && !StartupTask->IsDone())
+	{
+		return LOCTEXT("ProcessingAssetsStatus", "Processing assets...");
+	}
+
 	return FText::Format(LOCTEXT("ReconcileStatus", "Assets to check for reconcile: {0}"), FText::AsNumber(AddedAssetsCache.Num()));
 }
 
 bool FUncontrolledChangelistsModule::OnReconcileAssets()
 {
+	FScopedSlowTask Scope(0, LOCTEXT("ProcessingAssetsProgress", "Processing assets"));
+	const bool bShowCancelButton = false;
+	const bool bAllowInPIE = false;
+	Scope.MakeDialogDelayed(1.0f, bShowCancelButton, bAllowInPIE);
+
+	if (StartupTask)
+	{
+		while (!StartupTask->WaitCompletionWithTimeout(0.016))
+		{
+			Scope.EnterProgressFrame(0.f);
+		}
+				
+		AddedAssetsCache.Append(StartupTask->GetTask().GetAddedAssetsCache());
+		StartupTask = nullptr;
+	}
+
 	if ((!IsEnabled()) || AddedAssetsCache.IsEmpty())
 	{
 		return false;
 	}
 
+	Scope.EnterProgressFrame(0.f, LOCTEXT("ReconcileAssetsProgress", "Reconciling assets"));
+
 	CleanAssetsCaches();
 	bool bHasStateChanged = AddFilesToDefaultUncontrolledChangelist(AddedAssetsCache.Array(), FUncontrolledChangelistState::ECheckFlags::All);
 
@@ -148,21 +190,26 @@ bool FUncontrolledChangelistsModule::OnReconcileAssets()
 	return bHasStateChanged;
 }
 
-void FUncontrolledChangelistsModule::OnAssetAdded(const struct FAssetData& AssetData)
+void FUncontrolledChangelistsModule::OnAssetAdded(const FAssetData& AssetData)
 {
 	if (!IsEnabled())
 	{
 		return;
 	}
 
-	FPackagePath PackagePath;
+	OnAssetAddedInternal(AssetData, AddedAssetsCache, false);
+}
 
+void FUncontrolledChangelistsModule::OnAssetAddedInternal(const FAssetData& AssetData, TSet<FString>& InAddedAssetsCache, bool bInStartupTask)
+{
+	FPackagePath PackagePath;
 	if (!FPackagePath::TryFromPackageName(AssetData.PackageName, PackagePath))
 	{
 		return;
 	}
-	
-	if (!FPackageName::DoesPackageExist(PackagePath, &PackagePath))
+
+	// No need to check for existence when running startup task
+	if (!bInStartupTask && !FPackageName::DoesPackageExist(PackagePath, &PackagePath))
 	{
 		return; // If the package does not exist on disk there is nothing more to do
 	}
@@ -183,7 +230,7 @@ void FUncontrolledChangelistsModule::OnAssetAdded(const struct FAssetData& Asset
 
 	if (!IFileManager::Get().IsReadOnly(*Fullpath))
 	{
-		AddedAssetsCache.Add(MoveTemp(Fullpath));
+		InAddedAssetsCache.Add(MoveTemp(Fullpath));
 	}
 }
 
diff --git a/Engine/Source/Developer/UncontrolledChangelists/Public/UncontrolledChangelistsModule.h b/Engine/Source/Developer/UncontrolledChangelists/Public/UncontrolledChangelistsModule.h
--- a/Engine/Source/Developer/UncontrolledChangelists/Public/UncontrolledChangelistsModule.h
+++ b/Engine/Source/Developer/UncontrolledChangelists/Public/UncontrolledChangelistsModule.h
@@ -8,6 +8,9 @@
 #include "Modules/ModuleInterface.h"
 #include "Modules/ModuleManager.h"
 #include "UObject/UObjectGlobals.h"
+#include "Async/AsyncWork.h"
+
+struct FAssetData;
 
 /**
  * Interface for talking to Uncontrolled Changelists
@@ -78,7 +81,7 @@ public:
 	 * Delegate callback called when assets are added to AssetRegistry.
 	 * @param 	AssetData 	The asset just added.
 	 */
-	void OnAssetAdded(const struct FAssetData& AssetData);
+	void OnAssetAdded(const FAssetData& AssetData);
 
 	/**
 	 * Delegate callback called when an asset is loaded.
@@ -116,6 +119,15 @@ public:
 	void MoveFilesToControlledChangelist(const TArray<FString>& InUncontrolledFiles, const FSourceControlChangelistPtr& InChangelist);
 
 private:
+	/**
+	 * Helper use by Startup task and OnAssetLoaded delegate.
+	 * @param 	InAssetData 		The asset just added.
+	 * @param  	InAddedAssetsCache 	The cache to add the asset to.
+	 * @param	bInStartupTask		If true, this asset was added from the startup task.
+	 * 
+	 */
+	void OnAssetAddedInternal(const FAssetData& InAssetData, TSet<FString>& InAddedAssetsCache, bool bInStartupTask);
+
 	/**
 	 * Saves the state of UncontrolledChangelists to Json for persistency.
 	 */
@@ -161,8 +173,32 @@ private:
 	bool AddFilesToDefaultUncontrolledChangelist(const TArray<FString>& InFilenames, const FUncontrolledChangelistState::ECheckFlags InCheckFlags);
 
 private:
+	class FStartupTask : public FNonAbandonableTask
+	{
+	public:
+		
+		FStartupTask(FUncontrolledChangelistsModule* InOwner)
+			: Owner(InOwner)
+		{}
+		~FStartupTask() {}
+		
+		TStatId GetStatId() const
+		{
+			RETURN_QUICK_DECLARE_CYCLE_STAT(FStartupTask, STATGROUP_ThreadPoolAsyncTasks);
+		}
+		
+		void DoWork();
+		const TSet<FString>& GetAddedAssetsCache() const { return AddedAssetsCache; }
+	
+	private:
+		FUncontrolledChangelistsModule* Owner;
+		TSet<FString> AddedAssetsCache;
+	};
+
+	TUniquePtr<FAsyncTask<FStartupTask>> StartupTask;
 	FUncontrolledChangelistsStateCache	UncontrolledChangelistsStateCache;
 	TSet<FString>						AddedAssetsCache;
 	FDelegateHandle						OnAssetAddedDelegateHandle;
 	FDelegateHandle						OnObjectPreSavedDelegateHandle;
+	bool								bIsEnabled = false;
 };
