Uncontrolled Changelist: Filter out /Tmp assets + Procedural/Auto saves

diff --git a/Engine/Source/Developer/UncontrolledChangelists/Private/UncontrolledChangelistsModule.cpp b/Engine/Source/Developer/UncontrolledChangelists/Private/UncontrolledChangelistsModule.cpp
--- a/Engine/Source/Developer/UncontrolledChangelists/Private/UncontrolledChangelistsModule.cpp
+++ b/Engine/Source/Developer/UncontrolledChangelists/Private/UncontrolledChangelistsModule.cpp
@@ -209,9 +209,17 @@ void FUncontrolledChangelistsModule::OnAssetAddedInternal(const FAssetData& Asse
 	}
 
 	// No need to check for existence when running startup task
-	if (!bInStartupTask && !FPackageName::DoesPackageExist(PackagePath, &PackagePath))
+	if (!bInStartupTask)
 	{
-		return; // If the package does not exist on disk there is nothing more to do
+		if (FPackageName::IsTempPackage(PackagePath.GetPackageName()))
+		{
+			return; // Ignore temp packages
+		}
+
+		if(!FPackageName::DoesPackageExist(PackagePath, &PackagePath))
+		{
+			return; // If the package does not exist on disk there is nothing more to do
+		}
 	}
 
 	const FString LocalFullPath(PackagePath.GetLocalFullPath());
@@ -234,13 +242,25 @@ void FUncontrolledChangelistsModule::OnAssetAddedInternal(const FAssetData& Asse
 	}
 }
 
-void FUncontrolledChangelistsModule::OnObjectPreSaved(UObject* InAsset, const FObjectPreSaveContext& InPreSaveContext)
+void FUncontrolledChangelistsModule::OnObjectPreSaved(UObject* InObject, const FObjectPreSaveContext& InPreSaveContext)
 {
 	if (!IsEnabled())
 	{
 		return;
 	}
 
+	// Make sure we are catching the top level asset object to avoid processing same package multiple times
+	if (!InObject || !InObject->IsAsset())
+	{
+		return;
+	}
+
+	// Ignore procedural save and autosaves
+	if (InPreSaveContext.IsProceduralSave() || ((InPreSaveContext.GetSaveFlags() & SAVE_FromAutosave) != 0))
+	{
+		return;
+	}
+
 	FString Fullpath = FPaths::ConvertRelativePathToFull(InPreSaveContext.GetTargetFilename());
 
 	if (Fullpath.IsEmpty())
diff --git a/Engine/Source/Developer/UncontrolledChangelists/Public/UncontrolledChangelistsModule.h b/Engine/Source/Developer/UncontrolledChangelists/Public/UncontrolledChangelistsModule.h
--- a/Engine/Source/Developer/UncontrolledChangelists/Public/UncontrolledChangelistsModule.h
+++ b/Engine/Source/Developer/UncontrolledChangelists/Public/UncontrolledChangelistsModule.h
@@ -91,10 +91,10 @@ public:
 
 	/**
 	 * Delegate callback called before an asset has been written to disk.
-	 * @param 	InAsset 			The saved asset.
+	 * @param 	InObject 			The saved object.
 	 * @param 	InPreSaveContext 	Interface used to access saved parameters.
 	 */
-	void OnObjectPreSaved(UObject* InAsset, const FObjectPreSaveContext& InPreSaveContext);
+	void OnObjectPreSaved(UObject* InObject, const FObjectPreSaveContext& InPreSaveContext);
 
 	/**
 	 * Moves files to an Uncontrolled Changelist.
