Changes reverting UncontrolledChangelist new assets to also delete them.

diff --git a/Engine/Source/Developer/UncontrolledChangelists/Private/UncontrolledChangelistsModule.cpp b/Engine/Source/Developer/UncontrolledChangelists/Private/UncontrolledChangelistsModule.cpp
--- a/Engine/Source/Developer/UncontrolledChangelists/Private/UncontrolledChangelistsModule.cpp
+++ b/Engine/Source/Developer/UncontrolledChangelists/Private/UncontrolledChangelistsModule.cpp
@@ -242,6 +242,54 @@ void FUncontrolledChangelistsModule::OnAssetAddedInternal(const FAssetData& Asse
 	}
 }
 
+void FUncontrolledChangelistsModule::OnRevert(const TArray<FString>& InFilenames)
+{
+	if (!IsEnabled() || InFilenames.IsEmpty())
+	{
+		return;
+	}
+
+	ISourceControlProvider& SourceControlProvider = ISourceControlModule::Get().GetProvider();
+	IFileManager& FileManager = IFileManager::Get();
+	TArray<FSourceControlStateRef> UpdatedFilestates;
+
+	SourceControlProvider.GetState(InFilenames, UpdatedFilestates, EStateCacheUsage::ForceUpdate);
+
+	TArray<FString> FilesToDelete;
+	TArray<FString> FilesToRevert;
+
+	for (const FSourceControlStateRef& Filestate : UpdatedFilestates)
+	{
+		if (Filestate->IsSourceControlled())
+		{
+			FilesToRevert.Add(Filestate->GetFilename());
+		}
+		else
+		{
+			FilesToDelete.Add(Filestate->GetFilename());
+		}
+	}
+
+	if (!FilesToRevert.IsEmpty())
+	{
+		TSharedRef<FSync> ForceSyncOperation = ISourceControlOperation::Create<FSync>();
+		ForceSyncOperation->SetForce(true);
+		ForceSyncOperation->SetLastSyncedFlag(true);
+		SourceControlProvider.Execute(ForceSyncOperation, FilesToRevert);
+	}
+
+	for (const FString& FileToDelete : FilesToDelete)
+	{
+		const bool bRequireExists = true;
+		const bool bEvenReadOnly = false;
+		const bool bQuiet = false;
+
+		FileManager.Delete(*FileToDelete, bRequireExists, bEvenReadOnly, bQuiet);
+	}
+
+	UpdateStatus();
+}
+
 void FUncontrolledChangelistsModule::OnObjectPreSaved(UObject* InObject, const FObjectPreSaveContext& InPreSaveContext)
 {
 	if (!IsEnabled())
diff --git a/Engine/Source/Developer/UncontrolledChangelists/Public/UncontrolledChangelistsModule.h b/Engine/Source/Developer/UncontrolledChangelists/Public/UncontrolledChangelistsModule.h
--- a/Engine/Source/Developer/UncontrolledChangelists/Public/UncontrolledChangelistsModule.h
+++ b/Engine/Source/Developer/UncontrolledChangelists/Public/UncontrolledChangelistsModule.h
@@ -89,6 +89,11 @@ public:
 	 */
 	void OnAssetLoaded(UObject* InAsset);
 
+	/** Called when "Revert files" button is clicked. Reverts modified files and deletes new ones.
+	 *  @param	InFilenames		The files to be reverted
+	 */
+	void OnRevert(const TArray<FString>& InFilenames);
+
 	/**
 	 * Delegate callback called before an asset has been written to disk.
 	 * @param 	InObject 			The saved object.
diff --git a/Engine/Source/Editor/SourceControlWindows/Private/SSourceControlChangelists.cpp b/Engine/Source/Editor/SourceControlWindows/Private/SSourceControlChangelists.cpp
--- a/Engine/Source/Editor/SourceControlWindows/Private/SSourceControlChangelists.cpp
+++ b/Engine/Source/Editor/SourceControlWindows/Private/SSourceControlChangelists.cpp
@@ -774,12 +774,7 @@ void SSourceControlChangelistsWidget::OnRevert()
 	// Reverts selected Uncontrolled Files
 	if (!SelectedUncontrolledFiles.IsEmpty())
 	{
-		auto ForceSyncOperation = ISourceControlOperation::Create<FSync>();
-		ForceSyncOperation->SetForce(true);
-		ForceSyncOperation->SetLastSyncedFlag(true);
-		SourceControlProvider.Execute(ForceSyncOperation, SelectedUncontrolledFiles);
-
-		FUncontrolledChangelistsModule::Get().UpdateStatus();
+		FUncontrolledChangelistsModule::Get().OnRevert(SelectedUncontrolledFiles);
 	}
 }
 
