Uncontrolled Changelist: Make sure background task is completed on module shutdown

diff --git a/Engine/Source/Developer/UncontrolledChangelists/Private/UncontrolledChangelistsModule.cpp b/Engine/Source/Developer/UncontrolledChangelists/Private/UncontrolledChangelistsModule.cpp
--- a/Engine/Source/Developer/UncontrolledChangelists/Private/UncontrolledChangelistsModule.cpp
+++ b/Engine/Source/Developer/UncontrolledChangelists/Private/UncontrolledChangelistsModule.cpp
@@ -2,6 +2,7 @@
 
 #include "UncontrolledChangelistsModule.h"
 
+#include "CoreGlobals.h"
 #include "Algo/AnyOf.h"
 #include "Algo/Copy.h"
 #include "Algo/Find.h"
@@ -43,10 +44,14 @@ void FUncontrolledChangelistsModule::FStartupTask::DoWork()
 	TArray<FAssetData> Assets;
 	const bool bIncludeOnlyOnDiskAssets = true;
 	AssetRegistry.GetAllAssets(Assets, bIncludeOnlyOnDiskAssets);
-	Algo::ForEach(Assets, [this](const struct FAssetData& AssetData) 
+	for(const FAssetData& AssetData : Assets) 
 	{ 
+		if (IsEngineExitRequested())
+		{
+			break;
+		}
 		Owner->OnAssetAddedInternal(AssetData, AddedAssetsCache, true);
-	});
+	}
 	
 	UE_LOG(LogSourceControl, Log, TEXT("Uncontrolled asset enumeration finished in %s seconds (Found %d uncontrolled assets)"), *FString::SanitizeFloat(FPlatformTime::Seconds() - StartTime), AddedAssetsCache.Num());
 }
@@ -76,6 +81,12 @@ void FUncontrolledChangelistsModule::StartupModule()
 
 void FUncontrolledChangelistsModule::ShutdownModule()
 {
+	if (StartupTask)
+	{
+		StartupTask->EnsureCompletion();
+		StartupTask = nullptr;
+	}
+
 	FAssetRegistryModule* AssetRegistryModulePtr = static_cast<FAssetRegistryModule*>(FModuleManager::Get().GetModule(TEXT("AssetRegistry")));
 
 	// Check in case AssetRegistry has already been shutdown.
