Moves FUncontrolledChangelist::Name to FUncontrolledChangelistState::Description.

diff --git a/Engine/Source/Developer/UncontrolledChangelists/Private/UncontrolledChangelist.cpp b/Engine/Source/Developer/UncontrolledChangelists/Private/UncontrolledChangelist.cpp
--- a/Engine/Source/Developer/UncontrolledChangelists/Private/UncontrolledChangelist.cpp
+++ b/Engine/Source/Developer/UncontrolledChangelists/Private/UncontrolledChangelist.cpp
@@ -9,24 +9,20 @@
 #define LOCTEXT_NAMESPACE "UncontrolledChangelists"
 
 const FGuid FUncontrolledChangelist::DEFAULT_UNCONTROLLED_CHANGELIST_GUID = FGuid(0, 0, 0, 0);
-const FText FUncontrolledChangelist::DEFAULT_UNCONTROLLED_CHANGELIST_NAME = LOCTEXT("DefaultUncontrolledChangelist", "Default Uncontrolled Changelist");
 
 FUncontrolledChangelist::FUncontrolledChangelist()
 	: Guid(FGuid::NewGuid())
-	, Name()
 {
 }
 
-FUncontrolledChangelist::FUncontrolledChangelist(const FGuid& InGuid, const FString& InName)
+FUncontrolledChangelist::FUncontrolledChangelist(const FGuid& InGuid)
 	: Guid(InGuid)
-	, Name(InName)
 {
 }
 
 void FUncontrolledChangelist::Serialize(TSharedRef<FJsonObject> OutJsonObject) const
 {
 	OutJsonObject->SetStringField(GUID_NAME, Guid.ToString());
-	OutJsonObject->SetStringField(NAME_NAME, Name);
 }
 
 bool FUncontrolledChangelist::Deserialize(const TSharedRef<FJsonObject> InJsonValue)
@@ -45,12 +41,6 @@ bool FUncontrolledChangelist::Deserialize(const TSharedRef<FJsonObject> InJsonVa
 		return false;
 	}
 
-	if (!InJsonValue->TryGetStringField(NAME_NAME, Name))
-	{
-		UE_LOG(LogSourceControl, Error, TEXT("Cannot get field %s."), NAME_NAME);
-		return false;
-	}
-
 	return true;
 }
 
diff --git a/Engine/Source/Developer/UncontrolledChangelists/Private/UncontrolledChangelistState.cpp b/Engine/Source/Developer/UncontrolledChangelists/Private/UncontrolledChangelistState.cpp
--- a/Engine/Source/Developer/UncontrolledChangelists/Private/UncontrolledChangelistState.cpp
+++ b/Engine/Source/Developer/UncontrolledChangelists/Private/UncontrolledChangelistState.cpp
@@ -15,11 +15,19 @@
 
 #define LOCTEXT_NAMESPACE "UncontrolledChangelists"
 
+const FText FUncontrolledChangelistState::DEFAULT_UNCONTROLLED_CHANGELIST_DESCRIPTION = LOCTEXT("DefaultUncontrolledChangelist", "Default Uncontrolled Changelist");
+
 FUncontrolledChangelistState::FUncontrolledChangelistState(const FUncontrolledChangelist& InUncontrolledChangelist)
 	: Changelist(InUncontrolledChangelist)
 {
 }
 
+FUncontrolledChangelistState::FUncontrolledChangelistState(const FUncontrolledChangelist& InUncontrolledChangelist, const FText& InDescription)
+	: Changelist(InUncontrolledChangelist)
+{
+	SetDescription(InDescription);
+}
+
 FName FUncontrolledChangelistState::GetIconName() const
 {
 	return FName("SourceControl.UncontrolledChangelist");
@@ -30,14 +38,14 @@ FName FUncontrolledChangelistState::GetSmallIconName() const
 	return FName("SourceControl.UncontrolledChangelist_Small");
 }
 
-FText FUncontrolledChangelistState::GetDisplayText() const
+const FText& FUncontrolledChangelistState::GetDisplayText() const
 {
-	return FText::FromString(Changelist.ToString());
+	return Description;
 }
 
-FText FUncontrolledChangelistState::GetDescriptionText() const
+const FText& FUncontrolledChangelistState::GetDescriptionText() const
 {
-	return FText::FromString(Description);
+	return Description;
 }
 
 FText FUncontrolledChangelistState::GetDisplayTooltip() const
@@ -64,6 +72,8 @@ void FUncontrolledChangelistState::Serialize(TSharedRef<FJsonObject> OutJsonObje
 {
 	TArray<TSharedPtr<FJsonValue>> FileValues;
 
+	OutJsonObject->SetStringField(DESCRIPTION_NAME, Description.ToString());
+
 	Algo::Transform(Files, FileValues, [](const FSourceControlStateRef& File) { return MakeShareable(new FJsonValueString(File->GetFilename())); });
 	Algo::Transform(OfflineFiles, FileValues, [](const FString& OfflineFile) { return MakeShareable(new FJsonValueString(OfflineFile)); });
 
@@ -73,6 +83,15 @@ void FUncontrolledChangelistState::Serialize(TSharedRef<FJsonObject> OutJsonObje
 bool FUncontrolledChangelistState::Deserialize(const TSharedRef<FJsonObject> InJsonValue)
 {
 	const TArray<TSharedPtr<FJsonValue>>* FileValues = nullptr;
+	FString TempString;
+
+	if (!InJsonValue->TryGetStringField(DESCRIPTION_NAME, TempString) && !InJsonValue->TryGetStringField(NAME_NAME, TempString))
+	{
+		UE_LOG(LogSourceControl, Error, TEXT("Cannot get field %s or %s."), DESCRIPTION_NAME, NAME_NAME);
+		return false;
+	}
+
+	SetDescription(FText::FromString(TempString));
 
 	if ((!InJsonValue->TryGetArrayField(FILES_NAME, FileValues)) || (FileValues == nullptr))
 	{
@@ -204,4 +223,20 @@ void FUncontrolledChangelistState::RemoveDuplicates(TSet<FString>& InOutAddedAss
 	}
 }
 
+void FUncontrolledChangelistState::SetDescription(const FText& InDescription)
+{
+	Description = InDescription;
+
+	if (Description.EqualToCaseIgnored(FText::FromString(TEXT("default"))))
+	{
+		Description = FText::Format(LOCTEXT("Default_Override", "{0} (Uncontrolled Changelist)"), Description);
+	}
+}
+
+
+bool FUncontrolledChangelistState::ContainsFiles() const
+{
+	return (!Files.IsEmpty()) || (!OfflineFiles.IsEmpty());
+}
+
 #undef LOCTEXT_NAMESPACE
diff --git a/Engine/Source/Developer/UncontrolledChangelists/Private/UncontrolledChangelistsModule.cpp b/Engine/Source/Developer/UncontrolledChangelists/Private/UncontrolledChangelistsModule.cpp
--- a/Engine/Source/Developer/UncontrolledChangelists/Private/UncontrolledChangelistsModule.cpp
+++ b/Engine/Source/Developer/UncontrolledChangelists/Private/UncontrolledChangelistsModule.cpp
@@ -66,8 +66,7 @@ void FUncontrolledChangelistsModule::StartupModule()
 	}
 
 	// Adds Default Uncontrolled Changelist if it is not already present.
-	FUncontrolledChangelist DefaultUncontrolledChangelist(FUncontrolledChangelist::DEFAULT_UNCONTROLLED_CHANGELIST_GUID, FUncontrolledChangelist::DEFAULT_UNCONTROLLED_CHANGELIST_NAME.ToString());
-	UncontrolledChangelistsStateCache.FindOrAdd(MoveTemp(DefaultUncontrolledChangelist), MakeShareable(new FUncontrolledChangelistState(DefaultUncontrolledChangelist)));
+	GetDefaultUncontrolledChangelistState();
 
 	LoadState();
 
@@ -118,24 +117,12 @@ TArray<FUncontrolledChangelistStateRef> FUncontrolledChangelistsModule::GetChang
 
 bool FUncontrolledChangelistsModule::OnMakeWritable(const TArray<FString>& InFilenames)
 {
-	bool bHasStateChanged = false;
-
 	if (!IsEnabled())
 	{
 		return false;
 	}
 
-	FUncontrolledChangelist DefaultUncontrolledChangelist(FUncontrolledChangelist::DEFAULT_UNCONTROLLED_CHANGELIST_GUID, FUncontrolledChangelist::DEFAULT_UNCONTROLLED_CHANGELIST_NAME.ToString());
-	FUncontrolledChangelistsStateCache::ValueType& UncontrolledChangelistState = UncontrolledChangelistsStateCache.FindOrAdd(MoveTemp(DefaultUncontrolledChangelist), MakeShareable(new FUncontrolledChangelistState(DefaultUncontrolledChangelist)));
-
-	bHasStateChanged = UncontrolledChangelistState->AddFiles(InFilenames, FUncontrolledChangelistState::ECheckFlags::NotCheckedOut);
-
-	if (bHasStateChanged)
-	{
-		OnStateChanged();
-	}
-
-	return bHasStateChanged;
+	return AddFilesToDefaultUncontrolledChangelist(InFilenames, FUncontrolledChangelistState::ECheckFlags::NotCheckedOut);
 }
 
 void FUncontrolledChangelistsModule::UpdateStatus()
@@ -516,8 +503,7 @@ bool FUncontrolledChangelistsModule::AddFilesToDefaultUncontrolledChangelist(con
 {
 	bool bHasStateChanged = false;
 
-	FUncontrolledChangelist DefaultUncontrolledChangelist(FUncontrolledChangelist::DEFAULT_UNCONTROLLED_CHANGELIST_GUID, FUncontrolledChangelist::DEFAULT_UNCONTROLLED_CHANGELIST_NAME.ToString());
-	FUncontrolledChangelistsStateCache::ValueType& UncontrolledChangelistState = UncontrolledChangelistsStateCache.FindOrAdd(MoveTemp(DefaultUncontrolledChangelist), MakeShareable(new FUncontrolledChangelistState(DefaultUncontrolledChangelist)));
+	FUncontrolledChangelistStateRef UncontrolledChangelistState = GetDefaultUncontrolledChangelistState();
 
 	// Try to add files, they will be added only if they pass the required checks
 	bHasStateChanged = UncontrolledChangelistState->AddFiles(InFilenames, InCheckFlags);
@@ -530,6 +516,19 @@ bool FUncontrolledChangelistsModule::AddFilesToDefaultUncontrolledChangelist(con
 	return bHasStateChanged;
 }
 
+FUncontrolledChangelistStateRef FUncontrolledChangelistsModule::GetDefaultUncontrolledChangelistState()
+{
+	FUncontrolledChangelist DefaultUncontrolledChangelist(FUncontrolledChangelist::DEFAULT_UNCONTROLLED_CHANGELIST_GUID);
+	FUncontrolledChangelistStateRef* DefaultUncontrolledChangelistState = UncontrolledChangelistsStateCache.Find(DefaultUncontrolledChangelist);
+
+	if (DefaultUncontrolledChangelistState != nullptr)
+	{
+		return *DefaultUncontrolledChangelistState;
+	}
+
+	return UncontrolledChangelistsStateCache.Emplace(MoveTemp(DefaultUncontrolledChangelist), MakeShared<FUncontrolledChangelistState>(DefaultUncontrolledChangelist, FUncontrolledChangelistState::DEFAULT_UNCONTROLLED_CHANGELIST_DESCRIPTION));
+}
+
 void FUncontrolledChangelistsModule::SaveState() const
 {
 	TSharedPtr<FJsonObject> RootObject = MakeShareable(new FJsonObject);
@@ -587,7 +586,7 @@ void FUncontrolledChangelistsModule::LoadState()
 		return;
 	}
 
-	if (VersionNumber != VERSION_NUMBER)
+	if (VersionNumber > VERSION_NUMBER)
 	{
 		UE_LOG(LogSourceControl, Error, TEXT("Version number is invalid (file: %u, current: %u)."), VersionNumber, VERSION_NUMBER);
 		return;
@@ -610,7 +609,7 @@ void FUncontrolledChangelistsModule::LoadState()
 			continue;
 		}
 
-		FUncontrolledChangelistsStateCache::ValueType& UncontrolledChangelistState = UncontrolledChangelistsStateCache.FindOrAdd(MoveTemp(TempKey), MakeShareable(new FUncontrolledChangelistState(TempKey)));
+		FUncontrolledChangelistsStateCache::ValueType& UncontrolledChangelistState = UncontrolledChangelistsStateCache.FindOrAdd(MoveTemp(TempKey), MakeShared<FUncontrolledChangelistState>(TempKey));
 
 		UncontrolledChangelistState->Deserialize(JsonObject);
 	}
diff --git a/Engine/Source/Developer/UncontrolledChangelists/Public/UncontrolledChangelist.h b/Engine/Source/Developer/UncontrolledChangelists/Public/UncontrolledChangelist.h
--- a/Engine/Source/Developer/UncontrolledChangelists/Public/UncontrolledChangelist.h
+++ b/Engine/Source/Developer/UncontrolledChangelists/Public/UncontrolledChangelist.h
@@ -11,19 +11,17 @@
 
 class FText;
 
-class FUncontrolledChangelist : public TSharedFromThis<FUncontrolledChangelist, ESPMode::ThreadSafe>
+class UNCONTROLLEDCHANGELISTS_API FUncontrolledChangelist : public TSharedFromThis<FUncontrolledChangelist, ESPMode::ThreadSafe>
 {
 public:
 	static constexpr const TCHAR* GUID_NAME = TEXT("guid");
-	static constexpr const TCHAR* NAME_NAME = TEXT("name");
 	static const FGuid DEFAULT_UNCONTROLLED_CHANGELIST_GUID;
-	static const FText DEFAULT_UNCONTROLLED_CHANGELIST_NAME;
 
 public:
 	FUncontrolledChangelist();
 	
-	FUncontrolledChangelist(const FGuid& InGuid, const FString& InName);
-
+	FUncontrolledChangelist(const FGuid& InGuid);
+	
 	bool operator==(const FUncontrolledChangelist& InOther) const
 	{
 		return Guid == InOther.Guid;
@@ -46,7 +44,7 @@ public:
 
 	FString ToString() const
 	{
-		return Name.IsEmpty() ? Guid.ToString() : Name;
+		return Guid.ToString();
 	}
 
 	/**
@@ -64,7 +62,7 @@ public:
 
 private:
 	FGuid		Guid;
-	FString		Name;
 };
 
+typedef TSharedPtr<FUncontrolledChangelist, ESPMode::ThreadSafe> FUncontrolledChangelistPtr;
 typedef TSharedRef<FUncontrolledChangelist, ESPMode::ThreadSafe> FUncontrolledChangelistRef;
diff --git a/Engine/Source/Developer/UncontrolledChangelists/Public/UncontrolledChangelistState.h b/Engine/Source/Developer/UncontrolledChangelists/Public/UncontrolledChangelistState.h
--- a/Engine/Source/Developer/UncontrolledChangelists/Public/UncontrolledChangelistState.h
+++ b/Engine/Source/Developer/UncontrolledChangelists/Public/UncontrolledChangelistState.h
@@ -20,10 +20,13 @@
 
 class FJsonObject;
 
-class FUncontrolledChangelistState : public TSharedFromThis<FUncontrolledChangelistState, ESPMode::ThreadSafe>
+class UNCONTROLLEDCHANGELISTS_API FUncontrolledChangelistState : public TSharedFromThis<FUncontrolledChangelistState, ESPMode::ThreadSafe>
 {
 public:
 	static constexpr const TCHAR* FILES_NAME = TEXT("files");
+	static constexpr const TCHAR* NAME_NAME = TEXT("name");
+	static constexpr const TCHAR* DESCRIPTION_NAME = TEXT("description");
+	static const FText DEFAULT_UNCONTROLLED_CHANGELIST_DESCRIPTION;
 
 	enum class ECheckFlags
 	{
@@ -41,37 +44,46 @@ public:
 	};
 
 public:
+	// An FUncontrolledChangelistState should always reference a given Changelist (with a valid GUID).
+	FUncontrolledChangelistState() = delete;
+
 	FUncontrolledChangelistState(const FUncontrolledChangelist& InUncontrolledChangelist);
+	
+	FUncontrolledChangelistState(const FUncontrolledChangelist& InUncontrolledChangelist, const FText& InDescription);
+
+	// Uncontrolled Changelist states are unique and non-copyable, should always be used by reference to preserve cache coherence.
+	FUncontrolledChangelistState(const FUncontrolledChangelistState& InUncontrolledChangelistState) = delete;
+	FUncontrolledChangelistState& operator=(const FUncontrolledChangelistState& InUncontrolledChangelistState) = delete;
 
 	/**
 	 * Get the name of the icon graphic we should use to display the state in a UI.
 	 * @returns the name of the icon to display
 	 */
-	UNCONTROLLEDCHANGELISTS_API FName GetIconName() const;
+	 FName GetIconName() const;
 
 	/**
 	 * Get the name of the small icon graphic we should use to display the state in a UI.
 	 * @returns the name of the icon to display
 	 */
-	UNCONTROLLEDCHANGELISTS_API FName GetSmallIconName() const;
+	FName GetSmallIconName() const;
 
 	/**
 	 * Get a text representation of the state
 	 * @returns	the text to display for this state
 	 */
-	UNCONTROLLEDCHANGELISTS_API FText GetDisplayText() const;
+	const FText& GetDisplayText() const;
 
 	/**
 	 * Get a text representation of the state
 	 * @returns	the text to display for this state
 	 */
-	UNCONTROLLEDCHANGELISTS_API FText GetDescriptionText() const;
+	const FText& GetDescriptionText() const;
 
 	/**
 	 * Get a tooltip to describe this state
 	 * @returns	the text to display for this states tooltip
 	 */
-	UNCONTROLLEDCHANGELISTS_API FText GetDisplayTooltip() const;
+	FText GetDisplayTooltip() const;
 
 	/**
 	 * Get the timestamp of the last update that was made to this state.
@@ -79,9 +91,9 @@ public:
 	 */
 	const FDateTime& GetTimeStamp() const;
 
-	UNCONTROLLEDCHANGELISTS_API const TSet<FSourceControlStateRef>& GetFilesStates() const;
+	const TSet<FSourceControlStateRef>& GetFilesStates() const;
 	
-	UNCONTROLLEDCHANGELISTS_API const TSet<FString>& GetOfflineFiles() const;
+	const TSet<FString>& GetOfflineFiles() const;
 
 	/**
 	 * Serialize the state of the Uncontrolled Changelist to a Json Object.
@@ -123,9 +135,18 @@ public:
 	 */
 	void RemoveDuplicates(TSet<FString>& InOutAddedAssets);
 
+	/**
+	 * Sets a new description for this Uncontrolled Changelist
+	 * @param	InDescription	The new description to set.
+	 */
+	void SetDescription(const FText& InDescription);
+
+	/** Returns true if the Uncontrolled Changelists contains either Files or OfflineFiles.	*/
+	bool ContainsFiles() const;
+
 public:
 	FUncontrolledChangelist Changelist;
-	FString Description;
+	FText Description;
 	TSet<FSourceControlStateRef> Files;
 	TSet<FString> OfflineFiles;
 
diff --git a/Engine/Source/Developer/UncontrolledChangelists/Public/UncontrolledChangelistsModule.h b/Engine/Source/Developer/UncontrolledChangelists/Public/UncontrolledChangelistsModule.h
--- a/Engine/Source/Developer/UncontrolledChangelists/Public/UncontrolledChangelistsModule.h
+++ b/Engine/Source/Developer/UncontrolledChangelists/Public/UncontrolledChangelistsModule.h
@@ -22,7 +22,7 @@ class UNCONTROLLEDCHANGELISTS_API FUncontrolledChangelistsModule : public IModul
 public:	
 	static constexpr const TCHAR* VERSION_NAME = TEXT("version");
 	static constexpr const TCHAR* CHANGELISTS_NAME = TEXT("changelists");
-	static constexpr uint32 VERSION_NUMBER = 0;
+	static constexpr uint32 VERSION_NUMBER = 1;
 
 	/** Callback called when the state of the Uncontrolled Changelist Module (or any Uncontrolled Changelist) changed */
 	DECLARE_MULTICAST_DELEGATE(FOnUncontrolledChangelistModuleChanged);
@@ -173,6 +173,9 @@ private:
 	 */
 	bool AddFilesToDefaultUncontrolledChangelist(const TArray<FString>& InFilenames, const FUncontrolledChangelistState::ECheckFlags InCheckFlags);
 
+	/** Returns the default Uncontrolled Changelist state, creates it if it does not exist. */
+	FUncontrolledChangelistStateRef GetDefaultUncontrolledChangelistState();
+
 private:
 	class FStartupTask : public FNonAbandonableTask
 	{
diff --git a/Engine/Source/Editor/SourceControlWindows/Private/SSourceControlChangelistRows.cpp b/Engine/Source/Editor/SourceControlWindows/Private/SSourceControlChangelistRows.cpp
--- a/Engine/Source/Editor/SourceControlWindows/Private/SSourceControlChangelistRows.cpp
+++ b/Engine/Source/Editor/SourceControlWindows/Private/SSourceControlChangelistRows.cpp
@@ -169,14 +169,6 @@ void SUncontrolledChangelistTableRow::Construct(const FArguments& InArgs, const
 				SNew(STextBlock)
 				.Text(FText::Format(INVTEXT("({0})"), TreeItem->GetFileCount() + TreeItem->GetOfflineFileCount()))
 			]
-			+ SHorizontalBox::Slot()
-			.AutoWidth()
-			.Padding(2.0f, 0.0f)
-			.VAlign(VAlign_Center)
-			[
-				SNew(STextBlock)
-				.Text(this, &SUncontrolledChangelistTableRow::GetChangelistDescriptionText)
-			]
 		], InOwner);
 }
 
diff --git a/Engine/Source/Editor/SourceControlWindows/Private/SSourceControlChangelists.cpp b/Engine/Source/Editor/SourceControlWindows/Private/SSourceControlChangelists.cpp
--- a/Engine/Source/Editor/SourceControlWindows/Private/SSourceControlChangelists.cpp
+++ b/Engine/Source/Editor/SourceControlWindows/Private/SSourceControlChangelists.cpp
@@ -766,7 +766,7 @@ FSourceControlChangelistStatePtr SSourceControlChangelistsWidget::GetCurrentChan
 	return StaticCastSharedPtr<FChangelistTreeItem>(SelectedItems[0])->ChangelistState;
 }
 
-FUncontrolledChangelistStatePtr SSourceControlChangelistsWidget::GetCurrentUncontrolledChangelistState()
+FUncontrolledChangelistStatePtr SSourceControlChangelistsWidget::GetCurrentUncontrolledChangelistState() const
 {
 	if (!UncontrolledChangelistTreeView)
 	{
@@ -788,6 +788,12 @@ FSourceControlChangelistPtr SSourceControlChangelistsWidget::GetCurrentChangelis
 	return ChangelistState ? (FSourceControlChangelistPtr)(ChangelistState->GetChangelist()) : nullptr;
 }
 
+TOptional<FUncontrolledChangelist> SSourceControlChangelistsWidget::GetCurrentUncontrolledChangelist() const
+{
+	FUncontrolledChangelistStatePtr UncontrolledChangelistState = GetCurrentUncontrolledChangelistState();
+	return UncontrolledChangelistState ? UncontrolledChangelistState->Changelist : TOptional<FUncontrolledChangelist>();
+}
+
 FSourceControlChangelistStatePtr SSourceControlChangelistsWidget::GetChangelistStateFromSelection()
 {
 	TArray<FChangelistTreeItemPtr> SelectedItems = ChangelistTreeView->GetSelectedItems();
@@ -1804,7 +1810,7 @@ void SSourceControlChangelistsWidget::OnMoveFiles()
 		if (UncontrolledChangelist->GetTreeItemType() == IChangelistTreeItem::UncontrolledChangelist)
 		{
 			const TSharedPtr<FUncontrolledChangelistTreeItem>& TypedChangelist = StaticCastSharedPtr<FUncontrolledChangelistTreeItem>(UncontrolledChangelist);
-			Items.Emplace(TypedChangelist->GetDisplayText(), TypedChangelist->GetDescriptionText(), bCanEditAlreadyExistingChangelistDescription);
+			Items.Emplace(TypedChangelist->GetDisplayText(), FText(), bCanEditAlreadyExistingChangelistDescription);
 		}
 	}
 
diff --git a/Engine/Source/Editor/SourceControlWindows/Private/SSourceControlChangelists.h b/Engine/Source/Editor/SourceControlWindows/Private/SSourceControlChangelists.h
--- a/Engine/Source/Editor/SourceControlWindows/Private/SSourceControlChangelists.h
+++ b/Engine/Source/Editor/SourceControlWindows/Private/SSourceControlChangelists.h
@@ -69,8 +69,9 @@ private:
 
 	/** Returns the currently selected changelist state ptr or null in invalid cases */
 	FSourceControlChangelistStatePtr GetCurrentChangelistState();
-	FUncontrolledChangelistStatePtr GetCurrentUncontrolledChangelistState();
+	FUncontrolledChangelistStatePtr GetCurrentUncontrolledChangelistState() const;
 	FSourceControlChangelistPtr GetCurrentChangelist();
+	TOptional<FUncontrolledChangelist> GetCurrentUncontrolledChangelist() const;
 	FSourceControlChangelistStatePtr GetChangelistStateFromSelection();
 	FSourceControlChangelistPtr GetChangelistFromSelection();
 
