Fix UncontrolledChangelists warnings seen for cooked assets

diff --git a/Engine/Source/Developer/UncontrolledChangelists/Private/UncontrolledChangelistsModule.cpp b/Engine/Source/Developer/UncontrolledChangelists/Private/UncontrolledChangelistsModule.cpp
--- a/Engine/Source/Developer/UncontrolledChangelists/Private/UncontrolledChangelistsModule.cpp
+++ b/Engine/Source/Developer/UncontrolledChangelists/Private/UncontrolledChangelistsModule.cpp
@@ -197,6 +197,11 @@ void FUncontrolledChangelistsModule::OnAssetAdded(const FAssetData& AssetData)
 
 void FUncontrolledChangelistsModule::OnAssetAddedInternal(const FAssetData& AssetData, TSet<FString>& InAddedAssetsCache, bool bInStartupTask)
 {
+	if (AssetData.HasAnyPackageFlags(PKG_Cooked))
+	{
+		return;
+	}
+
 	FPackagePath PackagePath;
 	if (!FPackagePath::TryFromPackageName(AssetData.PackageName, PackagePath))
 	{
