Uncontrolled Changelist: deleted files support
- Deleted files will show up in the Uncontrolled Changelists when a Source Control provider is selected
- Deleted files won't show up if SCC is disabled but will still be tracked in case SCC is renabled

diff --git a/Engine/Source/Developer/UncontrolledChangelists/Private/UncontrolledChangelistState.cpp b/Engine/Source/Developer/UncontrolledChangelists/Private/UncontrolledChangelistState.cpp
--- a/Engine/Source/Developer/UncontrolledChangelists/Private/UncontrolledChangelistState.cpp
+++ b/Engine/Source/Developer/UncontrolledChangelists/Private/UncontrolledChangelistState.cpp
@@ -68,6 +68,28 @@ const TSet<FString>& FUncontrolledChangelistState::GetOfflineFiles() const
 	return OfflineFiles;
 }
 
+const TSet<FString>& FUncontrolledChangelistState::GetDeletedOfflineFiles() const
+{
+	return DeletedOfflineFiles;
+}
+
+int32 FUncontrolledChangelistState::GetFileCount() const
+{
+	// Avoid counting DeletedOfflineFiles.
+	return Files.Num() + OfflineFiles.Num();
+}
+
+TArray<FString> FUncontrolledChangelistState::GetFilenames() const
+{
+	TArray<FString> Filenames;
+	Filenames.Reserve(GetFileCount());
+
+	Algo::Transform(GetFilesStates(), Filenames, [](const TSharedRef<ISourceControlState>& FileState) { return FileState->GetFilename(); });
+	Algo::Transform(GetOfflineFiles(), Filenames, [](const FString& Pathname) { return Pathname; });
+
+	return Filenames;
+}
+
 void FUncontrolledChangelistState::Serialize(TSharedRef<FJsonObject> OutJsonObject) const
 {
 	TArray<TSharedPtr<FJsonValue>> FileValues;
@@ -76,6 +98,7 @@ void FUncontrolledChangelistState::Serialize(TSharedRef<FJsonObject> OutJsonObje
 
 	Algo::Transform(Files, FileValues, [](const FSourceControlStateRef& File) { return MakeShareable(new FJsonValueString(File->GetFilename())); });
 	Algo::Transform(OfflineFiles, FileValues, [](const FString& OfflineFile) { return MakeShareable(new FJsonValueString(OfflineFile)); });
+	Algo::Transform(DeletedOfflineFiles, FileValues, [](const FString& DeletedOfflineFile) { return MakeShareable(new FJsonValueString(DeletedOfflineFile)); });
 
 	OutJsonObject->SetArrayField(FILES_NAME, MoveTemp(FileValues));
 }
@@ -129,11 +152,22 @@ bool FUncontrolledChangelistState::AddFiles(const TArray<FString>& InFilenames,
 	if (!SourceControlProvider.IsAvailable())
 	{
 		int32 OldSize = OfflineFiles.Num();
-		Algo::CopyIf(SourceControlHelpers::AbsoluteFilenames(InFilenames), OfflineFiles, [&FileManager](const FString& Filename)
+		int32 OldDeletedSize = DeletedOfflineFiles.Num();
+
+		for (const FString& Filename : SourceControlHelpers::AbsoluteFilenames(InFilenames))
 		{
-			return FileManager.FileExists(*Filename);
-		});
-		return OldSize != OfflineFiles.Num();
+			if (FileManager.FileExists(*Filename))
+			{
+				OfflineFiles.Add(Filename);
+			}
+			else
+			{
+				// Keep in case we source control provider and can determine this file is source controlled
+				DeletedOfflineFiles.Add(Filename);
+			}
+		}
+				
+		return (OldSize != OfflineFiles.Num()) || (OldDeletedSize != DeletedOfflineFiles.Num());
 	}
 
 	if (bCheckStatus)
@@ -190,14 +224,17 @@ bool FUncontrolledChangelistState::UpdateStatus()
 {
 	TArray<FString> FilesToUpdate;
 	bool bOutChanged = false;
-	int32 InitialFileNumber = Files.Num();
-	int32 InitialOfflineFileNumber = OfflineFiles.Num();
+	const int32 InitialFileNumber = Files.Num();
+	const int32 InitialOfflineFileNumber = OfflineFiles.Num();
+	const int32 InitialDeletedOfflineFiles = DeletedOfflineFiles.Num();
 
 	Algo::Transform(Files, FilesToUpdate, [](const FSourceControlStateRef& State) { return State->GetFilename(); });
 	Algo::Copy(OfflineFiles, FilesToUpdate);
+	Algo::Copy(DeletedOfflineFiles, FilesToUpdate);
 
 	Files.Empty();
 	OfflineFiles.Empty();
+	DeletedOfflineFiles.Empty();
 
 	if (FilesToUpdate.Num() == 0)
 	{
@@ -206,10 +243,11 @@ bool FUncontrolledChangelistState::UpdateStatus()
 
 	bOutChanged |= AddFiles(FilesToUpdate, ECheckFlags::All);
 
-	bool bFileNumberChanged = InitialFileNumber == Files.Num();
-	bool bOfflineFileNumberChanged = InitialOfflineFileNumber == OfflineFiles.Num();
+	const bool bFileNumberChanged = InitialFileNumber == Files.Num();
+	const bool bOfflineFileNumberChanged = InitialOfflineFileNumber == OfflineFiles.Num();
+	const bool bDeletedOfflineFileNumberChanged = InitialDeletedOfflineFiles == DeletedOfflineFiles.Num();
 
-	bOutChanged |= bFileNumberChanged || bOfflineFileNumberChanged;
+	bOutChanged |= bFileNumberChanged || bOfflineFileNumberChanged || bDeletedOfflineFileNumberChanged;
 
 	return bOutChanged;
 }
@@ -237,7 +275,8 @@ void FUncontrolledChangelistState::SetDescription(const FText& InDescription)
 
 bool FUncontrolledChangelistState::ContainsFiles() const
 {
-	return (!Files.IsEmpty()) || (!OfflineFiles.IsEmpty());
+	// Ignore DeletedOfflineFiles
+	return !Files.IsEmpty() || !OfflineFiles.IsEmpty();
 }
 
 #undef LOCTEXT_NAMESPACE
diff --git a/Engine/Source/Developer/UncontrolledChangelists/Private/UncontrolledChangelistsModule.cpp b/Engine/Source/Developer/UncontrolledChangelists/Private/UncontrolledChangelistsModule.cpp
--- a/Engine/Source/Developer/UncontrolledChangelists/Private/UncontrolledChangelistsModule.cpp
+++ b/Engine/Source/Developer/UncontrolledChangelists/Private/UncontrolledChangelistsModule.cpp
@@ -132,6 +132,11 @@ bool FUncontrolledChangelistsModule::OnSaveWritable(const FString& InFilename)
 	return AddToUncontrolledChangelist({ InFilename });
 }
 
+bool FUncontrolledChangelistsModule::OnDeleteWritable(const FString& InFilename)
+{
+	return AddToUncontrolledChangelist({ InFilename });
+}
+
 bool FUncontrolledChangelistsModule::AddToUncontrolledChangelist(const TArray<FString>& InFilenames)
 {
 	if (!IsEnabled())
@@ -584,6 +589,9 @@ void FUncontrolledChangelistsModule::DeleteUncontrolledChangelist(const FUncontr
 		return;
 	}
 
+	// Get Deleted Offline files and move them to the Default UCL so that we don't lose them
+	GetDefaultUncontrolledChangelistState()->AddFiles((*UncontrolledChangelistState)->GetDeletedOfflineFiles().Array(), FUncontrolledChangelistState::ECheckFlags::None);
+
 	UncontrolledChangelistsStateCache.Remove(InUncontrolledChangelist);
 
 	OnStateChanged();
diff --git a/Engine/Source/Developer/UncontrolledChangelists/Public/UncontrolledChangelistState.h b/Engine/Source/Developer/UncontrolledChangelists/Public/UncontrolledChangelistState.h
--- a/Engine/Source/Developer/UncontrolledChangelists/Public/UncontrolledChangelistState.h
+++ b/Engine/Source/Developer/UncontrolledChangelists/Public/UncontrolledChangelistState.h
@@ -95,6 +95,18 @@ public:
 	
 	const TSet<FString>& GetOfflineFiles() const;
 
+	const TSet<FString>& GetDeletedOfflineFiles() const;
+
+	/**
+	 * Get the number of files in the CL. (Includes file states and offline files) 
+	 */
+	int32 GetFileCount() const;
+
+	/**
+	 * Get the filenames in the CL. (Includes file states and offline files)
+	 */
+	TArray<FString> GetFilenames() const;
+
 	/**
 	 * Serialize the state of the Uncontrolled Changelist to a Json Object.
 	 * @param 	OutJsonObject 	The Json object used to serialize.
@@ -149,6 +161,7 @@ public:
 	FText Description;
 	TSet<FSourceControlStateRef> Files;
 	TSet<FString> OfflineFiles;
+	TSet<FString> DeletedOfflineFiles;
 
 	/** The timestamp of the last update */
 	FDateTime TimeStamp;
diff --git a/Engine/Source/Developer/UncontrolledChangelists/Public/UncontrolledChangelistsModule.h b/Engine/Source/Developer/UncontrolledChangelists/Public/UncontrolledChangelistsModule.h
--- a/Engine/Source/Developer/UncontrolledChangelists/Public/UncontrolledChangelistsModule.h
+++ b/Engine/Source/Developer/UncontrolledChangelists/Public/UncontrolledChangelistsModule.h
@@ -57,6 +57,13 @@ public:
 	 */
 	bool OnSaveWritable(const FString& InFilename);
 
+	/**
+	 * Called when file has been deleted without an available Provider. Adds the file to the Default Uncontrolled Changelist
+	 * @param	InFilename		The file to be added.
+	 * @return	True if the file have been handled by the Uncontrolled module.
+	 */
+	bool OnDeleteWritable(const FString& InFilename);
+
 	/**
 	 * Called when files should have been marked for add without an available Provider. Adds the files to the Default Uncontrolled Changelist
 	 * @param	InFilenames		The files to be added.
diff --git a/Engine/Source/Editor/SourceControlWindows/Private/SSourceControlChangelistRows.cpp b/Engine/Source/Editor/SourceControlWindows/Private/SSourceControlChangelistRows.cpp
--- a/Engine/Source/Editor/SourceControlWindows/Private/SSourceControlChangelistRows.cpp
+++ b/Engine/Source/Editor/SourceControlWindows/Private/SSourceControlChangelistRows.cpp
@@ -200,7 +200,7 @@ void SUncontrolledChangelistTableRow::Construct(const FArguments& InArgs, const
 			.AutoWidth()
 			[
 				SNew(STextBlock)
-				.Text(FText::Format(INVTEXT("({0})"), TreeItem->GetFileCount() + TreeItem->GetOfflineFileCount()))
+				.Text(FText::Format(INVTEXT("({0})"), TreeItem->GetFileCount()))
 			]
 		], InOwner);
 }
diff --git a/Engine/Source/Editor/SourceControlWindows/Private/SSourceControlChangelists.cpp b/Engine/Source/Editor/SourceControlWindows/Private/SSourceControlChangelists.cpp
--- a/Engine/Source/Editor/SourceControlWindows/Private/SSourceControlChangelists.cpp
+++ b/Engine/Source/Editor/SourceControlWindows/Private/SSourceControlChangelists.cpp
@@ -679,8 +679,7 @@ void SSourceControlChangelistsWidget::RequestFileStatusRefresh(const IChangelist
 	else if (Changelist.GetTreeItemType() == IChangelistTreeItem::UncontrolledChangelist)
 	{
 		const FUncontrolledChangelistTreeItem& ChangelistItem = static_cast<const FUncontrolledChangelistTreeItem&>(Changelist);
-		Algo::Transform(ChangelistItem.UncontrolledChangelistState->GetFilesStates(), Pathnames, [](const TSharedRef<ISourceControlState>& FileState) { return FileState->GetFilename(); });
-		Algo::Transform(ChangelistItem.UncontrolledChangelistState->GetOfflineFiles(), Pathnames, [](const FString& Pathname) { return Pathname; });
+		Algo::Transform(ChangelistItem.UncontrolledChangelistState->GetFilenames(), Pathnames, [](const FString& Pathname) { return Pathname; });
 	}
 
 	if (!Pathnames.IsEmpty())
@@ -785,8 +784,7 @@ void SSourceControlChangelistsWidget::OnRefresh()
 
 	for (const TSharedRef<FUncontrolledChangelistState>& UncontrolledChangelistState : UncontrolledChangelistStates)
 	{
-		ElementsToProcess += UncontrolledChangelistState->GetFilesStates().Num();
-		ElementsToProcess += UncontrolledChangelistState->GetOfflineFiles().Num();
+		ElementsToProcess += UncontrolledChangelistState->GetFileCount();
 	}
 
 	FScopedSlowTask SlowTask(ElementsToProcess, LOCTEXT("SourceControl_RebuildTree", "Refreshing Tree Items"));
diff --git a/Engine/Source/Editor/SourceControlWindows/Private/SSourceControlCommon.h b/Engine/Source/Editor/SourceControlWindows/Private/SSourceControlCommon.h
--- a/Engine/Source/Editor/SourceControlWindows/Private/SSourceControlCommon.h
+++ b/Engine/Source/Editor/SourceControlWindows/Private/SSourceControlCommon.h
@@ -124,12 +124,7 @@ struct FUncontrolledChangelistTreeItem : public IChangelistTreeItem
 
 	int32 GetFileCount() const
 	{
-		return UncontrolledChangelistState->GetFilesStates().Num();
-	}
-
-	int32 GetOfflineFileCount() const
-	{
-		return UncontrolledChangelistState->GetOfflineFiles().Num();
+		return UncontrolledChangelistState->GetFileCount();
 	}
 
 	FUncontrolledChangelistStateRef UncontrolledChangelistState;
diff --git a/Engine/Source/Editor/UnrealEd/Private/ObjectTools.cpp b/Engine/Source/Editor/UnrealEd/Private/ObjectTools.cpp
--- a/Engine/Source/Editor/UnrealEd/Private/ObjectTools.cpp
+++ b/Engine/Source/Editor/UnrealEd/Private/ObjectTools.cpp
@@ -102,6 +102,7 @@
 #include "Containers/Set.h"
 #include "UObject/StrongObjectPtr.h"
 #include "Logging/LogMacros.h"
+#include "UncontrolledChangelistsModule.h"
 
 DEFINE_LOG_CATEGORY_STATIC(LogObjectTools, Log, All);
 
@@ -2358,6 +2359,7 @@ namespace ObjectTools
 		// Now delete all packages that have become empty
 		bool bMakeWritable = false;
 		bool bSilent = false;
+		bool bDeletedFileLocallyWritable = false;
 		TArray<FString> SCCFilesToRevert;
 		TArray<FString> SCCFilesToDelete;
 
@@ -2391,8 +2393,7 @@ namespace ObjectTools
 					}
 					else if (!bIsCheckedOut && SourceControlProvider.UsesLocalReadOnlyState() && !IFileManager::Get().IsReadOnly(*PackageFilename))
 					{
-						// The file was made writable and deleted, so we delete it on disk.
-						IFileManager::Get().Delete(*PackageFilename);
+						bDeletedFileLocallyWritable = true;
 					}
 					else
 					{
@@ -2433,16 +2434,24 @@ namespace ObjectTools
 					if(bMakeWritable || ReturnType == EAppReturnType::Yes)
 					{
 						FPlatformFileManager::Get().GetPlatformFile().SetReadOnly(*PackageFilename, false);
-						IFileManager::Get().Delete(*PackageFilename);
+						bDeletedFileLocallyWritable = true;
 					}
 				}
 				else
 				{
-					IFileManager::Get().Delete(*PackageFilename);
+					bDeletedFileLocallyWritable = true;
 				}
 			}
-		}
+		
+			if (bDeletedFileLocallyWritable)
+			{
+				FUncontrolledChangelistsModule& UncontrolledChangelistsModule = FUncontrolledChangelistsModule::Get();
+				UncontrolledChangelistsModule.OnDeleteWritable(PackageFilename);
 
+				IFileManager::Get().Delete(*PackageFilename);
+			}
+		}
+				
 		// Handle all source control revert and delete operations as a batched operation.
 		if (ISourceControlModule::Get().IsEnabled())
 		{
