Add insights scope to a function that requires optimization

diff --git a/Engine/Source/Developer/UncontrolledChangelists/Private/UncontrolledChangelistsModule.cpp b/Engine/Source/Developer/UncontrolledChangelists/Private/UncontrolledChangelistsModule.cpp
--- a/Engine/Source/Developer/UncontrolledChangelists/Private/UncontrolledChangelistsModule.cpp
+++ b/Engine/Source/Developer/UncontrolledChangelists/Private/UncontrolledChangelistsModule.cpp
@@ -144,6 +144,8 @@ bool FUncontrolledChangelistsModule::AddToUncontrolledChangelist(const TArray<FS
 		return false;
 	}
 
+	TRACE_CPUPROFILER_EVENT_SCOPE(FUncontrolledChangelistsModule::AddToUncontrolledChangelist);
+
 	TArray<FString> FullPaths;
 	FullPaths.Reserve(InFilenames.Num());
 	Algo::Transform(InFilenames, FullPaths, [](const FString& Filename) { return FPaths::ConvertRelativePathToFull(Filename); });
