Uncontrolled Changelist: Optimization
- OnStateChanged sets a dirty flag, SaveState is performed on End of frame if needed.

#rnx

diff --git a/Engine/Source/Developer/UncontrolledChangelists/Private/UncontrolledChangelistsModule.cpp b/Engine/Source/Developer/UncontrolledChangelists/Private/UncontrolledChangelistsModule.cpp
--- a/Engine/Source/Developer/UncontrolledChangelists/Private/UncontrolledChangelistsModule.cpp
+++ b/Engine/Source/Developer/UncontrolledChangelists/Private/UncontrolledChangelistsModule.cpp
@@ -18,6 +18,7 @@
 #include "Misc/FileHelper.h"
 #include "Misc/Paths.h"
 #include "Misc/ScopedSlowTask.h"
+#include "Misc/CoreDelegates.h"
 #include "PackagesDialog.h"
 #include "Serialization/JsonSerializer.h"
 #include "Serialization/JsonWriter.h"
@@ -70,7 +71,8 @@ void FUncontrolledChangelistsModule::StartupModule()
 	IAssetRegistry& AssetRegistry = AssetRegistryModule.Get();
 	OnAssetAddedDelegateHandle = AssetRegistry.OnAssetAdded().AddLambda([](const struct FAssetData& AssetData) { Get().OnAssetAdded(AssetData); });
 	OnObjectPreSavedDelegateHandle = FCoreUObjectDelegates::OnObjectPreSave.AddLambda([](UObject* InAsset, const FObjectPreSaveContext& InPreSaveContext) { Get().OnObjectPreSaved(InAsset, InPreSaveContext); });
-		
+	OnEndFrameDelegateHandle = FCoreDelegates::OnEndFrame.AddLambda([]() { Get().OnEndFrame(); });
+
 	StartupTask = MakeUnique<FAsyncTask<FStartupTask>>(this);		
 	StartupTask->StartBackgroundTask();
 }
@@ -83,6 +85,12 @@ void FUncontrolledChangelistsModule::ShutdownModule()
 		StartupTask = nullptr;
 	}
 
+	if (bIsStateDirty)
+	{
+		SaveState();
+		bIsStateDirty = false;
+	}
+
 	FAssetRegistryModule* AssetRegistryModulePtr = static_cast<FAssetRegistryModule*>(FModuleManager::Get().GetModule(TEXT("AssetRegistry")));
 
 	// Check in case AssetRegistry has already been shutdown.
@@ -92,6 +100,7 @@ void FUncontrolledChangelistsModule::ShutdownModule()
 	}
 
 	FCoreUObjectDelegates::OnObjectPreSave.Remove(OnObjectPreSavedDelegateHandle);
+	FCoreDelegates::OnEndFrame.Remove(OnEndFrameDelegateHandle);
 }
 
 bool FUncontrolledChangelistsModule::IsEnabled() const
@@ -601,8 +610,17 @@ void FUncontrolledChangelistsModule::DeleteUncontrolledChangelist(const FUncontr
 
 void FUncontrolledChangelistsModule::OnStateChanged()
 {
-	OnUncontrolledChangelistModuleChanged.Broadcast();
-	SaveState();
+	bIsStateDirty = true;
+}
+
+void FUncontrolledChangelistsModule::OnEndFrame()
+{
+	if (bIsStateDirty)
+	{
+		OnUncontrolledChangelistModuleChanged.Broadcast();
+		SaveState();
+		bIsStateDirty = false;
+	}
 }
 
 void FUncontrolledChangelistsModule::CleanAssetsCaches()
diff --git a/Engine/Source/Developer/UncontrolledChangelists/Public/UncontrolledChangelistsModule.h b/Engine/Source/Developer/UncontrolledChangelists/Public/UncontrolledChangelistsModule.h
--- a/Engine/Source/Developer/UncontrolledChangelists/Public/UncontrolledChangelistsModule.h
+++ b/Engine/Source/Developer/UncontrolledChangelists/Public/UncontrolledChangelistsModule.h
@@ -191,6 +191,11 @@ private:
 	 * Restores the previously saved state from Json.
 	 */
 	void LoadState();
+		
+	/**
+	 * Called on End of frame. Calls SaveState if needed.
+	 */
+	void OnEndFrame();
 
 	/**
 	 * Helper returning the location of the file used for persistency.
@@ -250,5 +255,7 @@ private:
 	TSet<FString>						AddedAssetsCache;
 	FDelegateHandle						OnAssetAddedDelegateHandle;
 	FDelegateHandle						OnObjectPreSavedDelegateHandle;
+	FDelegateHandle						OnEndFrameDelegateHandle;
 	bool								bIsEnabled = false;
+	bool								bIsStateDirty = false;
 };
