#UE Disable expensive uncontrolled changelists asset registry OnAssetAdded function

diff --git a/Engine/Source/Developer/UncontrolledChangelists/Private/UncontrolledChangelistsModule.cpp b/Engine/Source/Developer/UncontrolledChangelists/Private/UncontrolledChangelistsModule.cpp
--- a/Engine/Source/Developer/UncontrolledChangelists/Private/UncontrolledChangelistsModule.cpp
+++ b/Engine/Source/Developer/UncontrolledChangelists/Private/UncontrolledChangelistsModule.cpp
@@ -69,7 +69,7 @@ void FUncontrolledChangelistsModule::StartupModule()
 
 	FAssetRegistryModule& AssetRegistryModule = FModuleManager::GetModuleChecked<FAssetRegistryModule>(TEXT("AssetRegistry"));
 	IAssetRegistry& AssetRegistry = AssetRegistryModule.Get();
-	OnAssetAddedDelegateHandle = AssetRegistry.OnAssetAdded().AddLambda([](const struct FAssetData& AssetData) { Get().OnAssetAdded(AssetData); });
+	//OnAssetAddedDelegateHandle = AssetRegistry.OnAssetAdded().AddLambda([](const struct FAssetData& AssetData) { Get().OnAssetAdded(AssetData); });
 	OnObjectPreSavedDelegateHandle = FCoreUObjectDelegates::OnObjectPreSave.AddLambda([](UObject* InAsset, const FObjectPreSaveContext& InPreSaveContext) { Get().OnObjectPreSaved(InAsset, InPreSaveContext); });
 	OnEndFrameDelegateHandle = FCoreDelegates::OnEndFrame.AddLambda([]() { Get().OnEndFrame(); });
 
