Optimization for getting the uncontrolled file status in the Outliner

diff --git a/Engine/Source/Developer/UncontrolledChangelists/Private/UncontrolledChangelistState.cpp b/Engine/Source/Developer/UncontrolledChangelists/Private/UncontrolledChangelistState.cpp
--- a/Engine/Source/Developer/UncontrolledChangelists/Private/UncontrolledChangelistState.cpp
+++ b/Engine/Source/Developer/UncontrolledChangelists/Private/UncontrolledChangelistState.cpp
@@ -2,6 +2,7 @@
 
 #include "UncontrolledChangelistState.h"
 
+#include "Algo/AnyOf.h"
 #include "Algo/Copy.h"
 #include "Algo/Transform.h"
 #include "Dom/JsonObject.h"
@@ -90,6 +91,11 @@ TArray<FString> FUncontrolledChangelistState::GetFilenames() const
 	return Filenames;
 }
 
+bool FUncontrolledChangelistState::ContainsFilename(const FString& PackageFilename) const
+{
+	return  OfflineFiles.Contains(PackageFilename) || Algo::AnyOf(Files, [&PackageFilename](const TSharedRef<ISourceControlState>& FileState){ return FileState->GetFilename() == PackageFilename; });
+}
+
 void FUncontrolledChangelistState::Serialize(TSharedRef<FJsonObject> OutJsonObject) const
 {
 	TArray<TSharedPtr<FJsonValue>> FileValues;
diff --git a/Engine/Source/Developer/UncontrolledChangelists/Public/UncontrolledChangelistState.h b/Engine/Source/Developer/UncontrolledChangelists/Public/UncontrolledChangelistState.h
--- a/Engine/Source/Developer/UncontrolledChangelists/Public/UncontrolledChangelistState.h
+++ b/Engine/Source/Developer/UncontrolledChangelists/Public/UncontrolledChangelistState.h
@@ -107,6 +107,11 @@ public:
 	 */
 	TArray<FString> GetFilenames() const;
 
+	/**
+	 * Check whether a file exists in the file states or offline files
+	 */
+	bool ContainsFilename(const FString& PackageFilename) const;
+
 	/**
 	 * Serialize the state of the Uncontrolled Changelist to a Json Object.
 	 * @param 	OutJsonObject 	The Json object used to serialize.
diff --git a/Engine/Source/Editor/SceneOutliner/Private/SceneOutlinerTreeItemSCC.cpp b/Engine/Source/Editor/SceneOutliner/Private/SceneOutlinerTreeItemSCC.cpp
--- a/Engine/Source/Editor/SceneOutliner/Private/SceneOutlinerTreeItemSCC.cpp
+++ b/Engine/Source/Editor/SceneOutliner/Private/SceneOutlinerTreeItemSCC.cpp
@@ -181,7 +181,7 @@ void FSceneOutlinerTreeItemSCC::HandleUncontrolledChangelistsStateChanged()
 
 	for (const TSharedRef<FUncontrolledChangelistState>& UncontrolledChangelistStateRef : UncontrolledChangelistStates)
 	{
-		if (UncontrolledChangelistStateRef->GetFilenames().Contains(ExternalPackageFileName))
+		if (UncontrolledChangelistStateRef->ContainsFilename(ExternalPackageFileName))
 		{
 			UncontrolledChangelistState = UncontrolledChangelistStateRef;
 			break;
