Fix FUncontrolledChangelistsModule::OnAssetAddedInternal()

PR #10272

diff --git a/Engine/Source/Developer/UncontrolledChangelists/Private/UncontrolledChangelistsModule.cpp b/Engine/Source/Developer/UncontrolledChangelists/Private/UncontrolledChangelistsModule.cpp
--- a/Engine/Source/Developer/UncontrolledChangelists/Private/UncontrolledChangelistsModule.cpp
+++ b/Engine/Source/Developer/UncontrolledChangelists/Private/UncontrolledChangelistsModule.cpp
@@ -283,7 +283,7 @@ void FUncontrolledChangelistsModule::OnAssetAddedInternal(const FAssetData& Asse
 		return;
 	}
 
-	if (!IFileManager::Get().IsReadOnly(*Fullpath))
+	if (ISourceControlModule::Get().GetProvider().UsesLocalReadOnlyState() && !IFileManager::Get().IsReadOnly(*Fullpath))
 	{
 		InAddedAssetsCache.Add(MoveTemp(Fullpath));
 	}
