[UncontrolledChangelists] Removed public dependency on UnrealEd to break circular dependency "UncontrolledChangelists -> UnrealEd -> UncontrolledChangelists"

diff --git a/Engine/Source/Developer/UncontrolledChangelists/UncontrolledChangelists.Build.cs b/Engine/Source/Developer/UncontrolledChangelists/UncontrolledChangelists.Build.cs
--- a/Engine/Source/Developer/UncontrolledChangelists/UncontrolledChangelists.Build.cs
+++ b/Engine/Source/Developer/UncontrolledChangelists/UncontrolledChangelists.Build.cs
@@ -32,7 +32,6 @@ public class UncontrolledChangelists : ModuleRules
 				{
 					"EditorFramework",
 					"Engine",
-					"UnrealEd",
 					"PackagesDialog",
 				}
 			);
