SceneOutlinerTreeItemSCC: Fix a crash on shutdown when an item tried to unregister from the UncontrolledChangelists module in its destructor after the module was unloaded

diff --git a/Engine/Source/Developer/UncontrolledChangelists/Public/UncontrolledChangelistsModule.h b/Engine/Source/Developer/UncontrolledChangelists/Public/UncontrolledChangelistsModule.h
--- a/Engine/Source/Developer/UncontrolledChangelists/Public/UncontrolledChangelistsModule.h
+++ b/Engine/Source/Developer/UncontrolledChangelists/Public/UncontrolledChangelistsModule.h
@@ -81,9 +81,24 @@ public:
 	 * @return A reference to the UncontrolledChangelists module.
 	 */
 	static inline FUncontrolledChangelistsModule& Get()
+	{
+		return FModuleManager::LoadModuleChecked<FUncontrolledChangelistsModule>(GetModuleName());
+	}
+
+	/**
+	 * Checks to see if this module is loaded and ready.  It is only valid to call Get() if IsAvailable() returns true.
+	 *
+	 * @return True if the module is loaded and ready to use
+	 */
+	static inline bool IsAvailable()
+	{
+		return FModuleManager::Get().IsModuleLoaded( GetModuleName() );
+	}
+
+	static FName GetModuleName()
 	{
 		static FName UncontrolledChangelistsModuleName("UncontrolledChangelists");
-		return FModuleManager::LoadModuleChecked<FUncontrolledChangelistsModule>(UncontrolledChangelistsModuleName);
+		return UncontrolledChangelistsModuleName;
 	}
 
 	/**
diff --git a/Engine/Source/Editor/SceneOutliner/Private/SceneOutlinerTreeItemSCC.cpp b/Engine/Source/Editor/SceneOutliner/Private/SceneOutlinerTreeItemSCC.cpp
--- a/Engine/Source/Editor/SceneOutliner/Private/SceneOutlinerTreeItemSCC.cpp
+++ b/Engine/Source/Editor/SceneOutliner/Private/SceneOutlinerTreeItemSCC.cpp
@@ -20,8 +20,11 @@ FSceneOutlinerTreeItemSCC::FSceneOutlinerTreeItemSCC(FSceneOutlinerTreeItemPtr I
 
 FSceneOutlinerTreeItemSCC::~FSceneOutlinerTreeItemSCC()
 {
-	FUncontrolledChangelistsModule& UncontrolledChangelistModule = FUncontrolledChangelistsModule::Get();
-	UncontrolledChangelistModule.OnUncontrolledChangelistModuleChanged.Remove(UncontrolledChangelistChangedHandle);
+	if(FUncontrolledChangelistsModule::IsAvailable())
+	{
+		FUncontrolledChangelistsModule& UncontrolledChangelistModule = FUncontrolledChangelistsModule::Get();
+		UncontrolledChangelistModule.OnUncontrolledChangelistModuleChanged.Remove(UncontrolledChangelistChangedHandle);
+	}
 	
 	DisconnectSourceControl();
 }
