SCC: Fix Uncontrolled Changelist reconcile
- Make sure we don't register to OnAssetAdded before the initial AR Scan completes
- StartupTask will only run after initial AR Scan catching all on disk assets

diff --git a/Engine/Source/Developer/UncontrolledChangelists/Private/UncontrolledChangelistsModule.cpp b/Engine/Source/Developer/UncontrolledChangelists/Private/UncontrolledChangelistsModule.cpp
--- a/Engine/Source/Developer/UncontrolledChangelists/Private/UncontrolledChangelistsModule.cpp
+++ b/Engine/Source/Developer/UncontrolledChangelists/Private/UncontrolledChangelistsModule.cpp
@@ -27,6 +27,7 @@
 #include "SourceControlPreferences.h"
 #include "Styling/SlateTypes.h"
 #include "UObject/ObjectSaveContext.h"
+#include "Engine/AssetManager.h"
 
 #define LOCTEXT_NAMESPACE "UncontrolledChangelists"
 
@@ -67,18 +68,37 @@ void FUncontrolledChangelistsModule::StartupModule()
 
 	LoadState();
 
-	FAssetRegistryModule& AssetRegistryModule = FModuleManager::GetModuleChecked<FAssetRegistryModule>(TEXT("AssetRegistry"));
-	IAssetRegistry& AssetRegistry = AssetRegistryModule.Get();
-	//OnAssetAddedDelegateHandle = AssetRegistry.OnAssetAdded().AddLambda([](const struct FAssetData& AssetData) { Get().OnAssetAdded(AssetData); });
 	OnObjectPreSavedDelegateHandle = FCoreUObjectDelegates::OnObjectPreSave.AddLambda([](UObject* InAsset, const FObjectPreSaveContext& InPreSaveContext) { Get().OnObjectPreSaved(InAsset, InPreSaveContext); });
 	OnEndFrameDelegateHandle = FCoreDelegates::OnEndFrame.AddLambda([]() { Get().OnEndFrame(); });
 
-	StartupTask = MakeUnique<FAsyncTask<FStartupTask>>(this);		
-	StartupTask->StartBackgroundTask();
+	// Create initial scan event object
+	InitialScanEvent = MakeShared<FInitialScanEvent>();
+
+	UAssetManager::CallOrRegister_OnCompletedInitialScan(FSimpleMulticastDelegate::FDelegate::CreateLambda([this, WeakScanEvent = InitialScanEvent->AsWeak()]()
+	{
+		// Weak here allows us to check if module as been shutdown before using [this]
+		if(!WeakScanEvent.IsValid())
+		{
+			return;
+		}
+
+		StartupTask = MakeUnique<FAsyncTask<FStartupTask>>(this);
+		StartupTask->StartBackgroundTask();
+
+		InitialScanEvent = nullptr;
+
+		FAssetRegistryModule& AssetRegistryModule = FModuleManager::GetModuleChecked<FAssetRegistryModule>(TEXT("AssetRegistry"));
+		IAssetRegistry& AssetRegistry = AssetRegistryModule.Get();
+		OnAssetAddedDelegateHandle = AssetRegistry.OnAssetAdded().AddLambda([](const struct FAssetData& AssetData) { Get().OnAssetAdded(AssetData); });
+	}));
+
 }
 
 void FUncontrolledChangelistsModule::ShutdownModule()
 {
+	// This will make sure callback for initial scan early outs if module was shutdown
+	InitialScanEvent = nullptr;
+
 	if (StartupTask)
 	{
 		StartupTask->EnsureCompletion();
@@ -191,6 +211,11 @@ void FUncontrolledChangelistsModule::UpdateStatus()
 
 FText FUncontrolledChangelistsModule::GetReconcileStatus() const
 {
+	if (InitialScanEvent.IsValid())
+	{
+		return LOCTEXT("WaitForAssetRegistryStatus", "Waiting for Asset Registry initial scan...");
+	}
+
 	if (StartupTask && !StartupTask->IsDone())
 	{
 		return LOCTEXT("ProcessingAssetsStatus", "Processing assets...");
@@ -649,6 +674,12 @@ void FUncontrolledChangelistsModule::OnStateChanged()
 
 void FUncontrolledChangelistsModule::OnEndFrame()
 {
+	if (StartupTask && StartupTask->IsDone())
+	{
+		AddedAssetsCache.Append(StartupTask->GetTask().GetAddedAssetsCache());
+		StartupTask = nullptr;
+	}
+
 	if (bIsStateDirty)
 	{
 		OnUncontrolledChangelistModuleChanged.Broadcast();
diff --git a/Engine/Source/Developer/UncontrolledChangelists/Public/UncontrolledChangelistsModule.h b/Engine/Source/Developer/UncontrolledChangelists/Public/UncontrolledChangelistsModule.h
--- a/Engine/Source/Developer/UncontrolledChangelists/Public/UncontrolledChangelistsModule.h
+++ b/Engine/Source/Developer/UncontrolledChangelists/Public/UncontrolledChangelistsModule.h
@@ -272,6 +272,10 @@ private:
 		TSet<FString> AddedAssetsCache;
 	};
 
+	// Used to determine if the initial Asset Registry scan was completed or the module was shutdown
+	struct FInitialScanEvent : public TSharedFromThis<FInitialScanEvent> {};
+	TSharedPtr<FInitialScanEvent> InitialScanEvent;
+
 	TUniquePtr<FAsyncTask<FStartupTask>> StartupTask;
 	FUncontrolledChangelistsStateCache	UncontrolledChangelistsStateCache;
 	TSet<FString>						AddedAssetsCache;
