Fixed crash with UncontrolledChangelists
- Added a sanitization pass when loading the json from disk and when writing it to disk to enforce one file being part of one cl.

diff --git a/Engine/Source/Developer/UncontrolledChangelists/Private/UncontrolledChangelistState.cpp b/Engine/Source/Developer/UncontrolledChangelists/Private/UncontrolledChangelistState.cpp
--- a/Engine/Source/Developer/UncontrolledChangelists/Private/UncontrolledChangelistState.cpp
+++ b/Engine/Source/Developer/UncontrolledChangelists/Private/UncontrolledChangelistState.cpp
@@ -266,6 +266,16 @@ void FUncontrolledChangelistState::RemoveDuplicates(TSet<FString>& InOutAddedAss
 		
 		InOutAddedAssets.Remove(Filename);
 	}
+
+	for (const FString& OfflineFile : OfflineFiles)
+	{
+		InOutAddedAssets.Remove(OfflineFile);
+	}
+
+	for (const FString& DeletedOfflineFile : DeletedOfflineFiles)
+	{
+		InOutAddedAssets.Remove(DeletedOfflineFile);
+	}
 }
 
 void FUncontrolledChangelistState::SetDescription(const FText& InDescription)
diff --git a/Engine/Source/Developer/UncontrolledChangelists/Private/UncontrolledChangelistsModule.cpp b/Engine/Source/Developer/UncontrolledChangelists/Private/UncontrolledChangelistsModule.cpp
--- a/Engine/Source/Developer/UncontrolledChangelists/Private/UncontrolledChangelistsModule.cpp
+++ b/Engine/Source/Developer/UncontrolledChangelists/Private/UncontrolledChangelistsModule.cpp
@@ -728,8 +728,10 @@ FUncontrolledChangelistStateRef FUncontrolledChangelistsModule::GetDefaultUncont
 	return UncontrolledChangelistsStateCache.Emplace(MoveTemp(DefaultUncontrolledChangelist), MakeShared<FUncontrolledChangelistState>(DefaultUncontrolledChangelist, FUncontrolledChangelistState::DEFAULT_UNCONTROLLED_CHANGELIST_DESCRIPTION));
 }
 
-void FUncontrolledChangelistsModule::SaveState() const
+void FUncontrolledChangelistsModule::SaveState()
 {
+	SanitizeState();
+
 	TSharedPtr<FJsonObject> RootObject = MakeShareable(new FJsonObject);
 	TArray<TSharedPtr<FJsonValue>> UncontrolledChangelistsArray;
 
@@ -813,9 +815,55 @@ void FUncontrolledChangelistsModule::LoadState()
 		UncontrolledChangelistState->Deserialize(JsonObject);
 	}
 
+	SanitizeState();
+
 	UE_LOG(LogSourceControl, Display, TEXT("Uncontrolled Changelist persistency file loaded %s"), *GetPersistentFilePath());
 }
 
+void FUncontrolledChangelistsModule::SanitizeState()
+{
+	TSet<FString> AllFiles;
+
+	for (const TPair<FUncontrolledChangelist, FUncontrolledChangelistStateRef>& Pair : UncontrolledChangelistsStateCache)
+	{
+		FUncontrolledChangelistStateRef UncontrolledChangelistState = Pair.Value;
+
+		// UncontrolledChangelistState->Files
+		{
+			for (TSet<FSourceControlStateRef>::TIterator FileStateIt = UncontrolledChangelistState->Files.CreateIterator(); FileStateIt; ++FileStateIt)
+			{
+				if (AllFiles.Contains((*FileStateIt)->GetFilename()))
+				{
+					FileStateIt.RemoveCurrent();
+				}
+				else
+				{
+					AllFiles.Add((*FileStateIt)->GetFilename());
+				}
+			}
+
+		}
+
+		auto RemoveDuplicateFiles = [&AllFiles](TSet<FString>& Files)
+		{
+			for (TSet<FString>::TIterator FileIt = Files.CreateIterator(); FileIt; ++FileIt)
+			{
+				if (AllFiles.Contains(*FileIt))
+				{
+					FileIt.RemoveCurrent();
+				}
+				else
+				{
+					AllFiles.Add(*FileIt);
+				}
+			}
+		};
+
+		RemoveDuplicateFiles(UncontrolledChangelistState->OfflineFiles);
+		RemoveDuplicateFiles(UncontrolledChangelistState->DeletedOfflineFiles);
+	}
+}
+
 FString FUncontrolledChangelistsModule::GetPersistentFilePath() const
 {
 	return FPaths::ProjectSavedDir() + TEXT("SourceControl/UncontrolledChangelists.json");
diff --git a/Engine/Source/Developer/UncontrolledChangelists/Public/UncontrolledChangelistsModule.h b/Engine/Source/Developer/UncontrolledChangelists/Public/UncontrolledChangelistsModule.h
--- a/Engine/Source/Developer/UncontrolledChangelists/Public/UncontrolledChangelistsModule.h
+++ b/Engine/Source/Developer/UncontrolledChangelists/Public/UncontrolledChangelistsModule.h
@@ -207,13 +207,18 @@ private:
 	/**
 	 * Saves the state of UncontrolledChangelists to Json for persistency.
 	 */
-	void SaveState() const;
+	void SaveState();
 	
 	/**
 	 * Restores the previously saved state from Json.
 	 */
 	void LoadState();
-		
+	
+	/**
+	 * Removes any duplicated files across changelists
+	 */
+	void SanitizeState();
+
 	/**
 	 * Called on End of frame. Calls SaveState if needed.
 	 */
