Fix a potential shutdown crash in the uncontrolled changelist module.

- When FUncontrolledChangelistsModule::ShutdownModule is called we do try to ensure that the start up task has completed if it is still running but in very rare cases the job will not get time to run on worker threads until the engine exit path has already started to exit and things that it relies on, such as the asset registry have already shutdown. In these cases we can hit an assert.
- We now hook into the OnEnginePreExit delegate which is invoked before major engine systems start to shutdown.
- We also try to first cancel the job, in case it is still pending on the task system as that is the fastest way to stop it from running.
- If the job has already started we will have to wait for it to complete.
- Note that the job will try to early out once it detect that the engine is shutting down so the wait here should be relatively short.
- There is now no need to check that the start up task has finished during ::ShutdownModule so we now double check for future coding errors with an assert instead.

diff --git a/Engine/Source/Developer/UncontrolledChangelists/Private/UncontrolledChangelistsModule.cpp b/Engine/Source/Developer/UncontrolledChangelists/Private/UncontrolledChangelistsModule.cpp
--- a/Engine/Source/Developer/UncontrolledChangelists/Private/UncontrolledChangelistsModule.cpp
+++ b/Engine/Source/Developer/UncontrolledChangelists/Private/UncontrolledChangelistsModule.cpp
@@ -51,7 +51,7 @@ void FUncontrolledChangelistsModule::FStartupTask::DoWork()
 		Owner->OnAssetAddedInternal(AssetData, AddedAssetsCache, true);
 	}
 	
-	UE_LOG(LogSourceControl, Log, TEXT("Uncontrolled asset enumeration finished in %s seconds (Found %d uncontrolled assets)"), *FString::SanitizeFloat(FPlatformTime::Seconds() - StartTime), AddedAssetsCache.Num());
+	UE_LOG(LogSourceControl, Log, TEXT("Uncontrolled asset enumeration finished in %s seconds (Found %d uncontrolled assets)"), *FString::SanitizeFloat(FPlatformTime::Seconds() - StartTime), AddedAssetsCache.Num());;
 }
 
 void FUncontrolledChangelistsModule::StartupModule()
@@ -90,6 +90,27 @@ void FUncontrolledChangelistsModule::StartupModule()
 		FAssetRegistryModule& AssetRegistryModule = FModuleManager::GetModuleChecked<FAssetRegistryModule>(TEXT("AssetRegistry"));
 		IAssetRegistry& AssetRegistry = AssetRegistryModule.Get();
 		OnAssetAddedDelegateHandle = AssetRegistry.OnAssetAdded().AddLambda([](const struct FAssetData& AssetData) { Get().OnAssetAdded(AssetData); });
+	
+		OnEnginePreExitDelegateHandle = FCoreDelegates::OnEnginePreExit.AddLambda([this]()
+			{
+				if (StartupTask.IsValid())
+				{
+					if (StartupTask->Cancel())
+					{
+						UE_LOG(LogSourceControl, Log, TEXT("UncontrolledChangelistsModule start up task was cancelled due to engine exit"));
+					}
+					else  if (!StartupTask->IsDone())
+					{
+						const double StartTime = FPlatformTime::Seconds();
+
+						UE_LOG(LogSourceControl, Log, TEXT("Waiting on UncontrolledChangelistsModule start up task to complete..."));
+						StartupTask->EnsureCompletion();
+						UE_LOG(LogSourceControl, Log, TEXT("UncontrolledChangelistsModule start up task finished, exit was stalled for %.1f(s)"), FPlatformTime::Seconds() - StartTime);
+					}
+
+					StartupTask.Reset();
+				}
+			});
 	}));
 
 }
@@ -99,11 +120,7 @@ void FUncontrolledChangelistsModule::ShutdownModule()
 	// This will make sure callback for initial scan early outs if module was shutdown
 	InitialScanEvent = nullptr;
 
-	if (StartupTask)
-	{
-		StartupTask->EnsureCompletion();
-		StartupTask = nullptr;
-	}
+	checkf(StartupTask == nullptr, TEXT("The startup task should be cleaned up when OnEnginePreExit is called at the latest"));
 
 	if (bIsStateDirty)
 	{
@@ -118,9 +135,16 @@ void FUncontrolledChangelistsModule::ShutdownModule()
 	{
 		AssetRegistryModulePtr->Get().OnAssetAdded().Remove(OnAssetAddedDelegateHandle);
 	}
+	OnAssetAddedDelegateHandle.Reset();
 
 	FCoreUObjectDelegates::OnObjectPreSave.Remove(OnObjectPreSavedDelegateHandle);
+	OnObjectPreSavedDelegateHandle.Reset();
+
 	FCoreDelegates::OnEndFrame.Remove(OnEndFrameDelegateHandle);
+	OnEndFrameDelegateHandle.Reset();
+
+	FCoreDelegates::OnEnginePreExit.Remove(OnEnginePreExitDelegateHandle);
+	OnEnginePreExitDelegateHandle.Reset();
 }
 
 bool FUncontrolledChangelistsModule::IsEnabled() const
diff --git a/Engine/Source/Developer/UncontrolledChangelists/Public/UncontrolledChangelistsModule.h b/Engine/Source/Developer/UncontrolledChangelists/Public/UncontrolledChangelistsModule.h
--- a/Engine/Source/Developer/UncontrolledChangelists/Public/UncontrolledChangelistsModule.h
+++ b/Engine/Source/Developer/UncontrolledChangelists/Public/UncontrolledChangelistsModule.h
@@ -284,6 +284,7 @@ private:
 	TUniquePtr<FAsyncTask<FStartupTask>> StartupTask;
 	FUncontrolledChangelistsStateCache	UncontrolledChangelistsStateCache;
 	TSet<FString>						AddedAssetsCache;
+	FDelegateHandle						OnEnginePreExitDelegateHandle;
 	FDelegateHandle						OnAssetAddedDelegateHandle;
 	FDelegateHandle						OnObjectPreSavedDelegateHandle;
 	FDelegateHandle						OnEndFrameDelegateHandle;
