Updated uncontrolled changelists to respect ISourceControlModule::GetCustomProjects

When using custom projects:
  * Uncontrolled changelist data is stored in per-project JSON files.
  * Packages outside of any known project root are filtered out (eg, assets pending reconcile).

This uses ISourceControlModule::OnCustomProjectsChanged to detect when the list of custom projects changes, and in response will:
  * Save any pending uncontrolled changelist state (if dirty).
  * Discard any assets pending reconcile.
  * Discard any uncontrolled changelist state.
  * Load the uncontrolled changelist state for the new project(s).
  * Discover any assets pending reconcile for the new project(s).

This also updates the uncontrolled changelists module to dynamically initialize/shutdown correctly if the source control provider changes between one that does/doesn't support the uncontrolled changelists feature.

diff --git a/Engine/Source/Developer/UncontrolledChangelists/Private/UncontrolledChangelistState.cpp b/Engine/Source/Developer/UncontrolledChangelists/Private/UncontrolledChangelistState.cpp
--- a/Engine/Source/Developer/UncontrolledChangelists/Private/UncontrolledChangelistState.cpp
+++ b/Engine/Source/Developer/UncontrolledChangelists/Private/UncontrolledChangelistState.cpp
@@ -96,15 +96,20 @@ bool FUncontrolledChangelistState::ContainsFilename(const FString& PackageFilena
 	return  OfflineFiles.Contains(PackageFilename) || Algo::AnyOf(Files, [&PackageFilename](const TSharedRef<ISourceControlState>& FileState){ return FileState->GetFilename() == PackageFilename; });
 }
 
-void FUncontrolledChangelistState::Serialize(TSharedRef<FJsonObject> OutJsonObject) const
+void FUncontrolledChangelistState::Serialize(TSharedRef<FJsonObject> OutJsonObject, const TFunction<bool(const FString&)>& FilenameFilter) const
 {
 	TArray<TSharedPtr<FJsonValue>> FileValues;
 
 	OutJsonObject->SetStringField(DESCRIPTION_NAME, Description.ToString());
 
-	Algo::Transform(Files, FileValues, [](const FSourceControlStateRef& File) { return MakeShareable(new FJsonValueString(File->GetFilename())); });
-	Algo::Transform(OfflineFiles, FileValues, [](const FString& OfflineFile) { return MakeShareable(new FJsonValueString(OfflineFile)); });
-	Algo::Transform(DeletedOfflineFiles, FileValues, [](const FString& DeletedOfflineFile) { return MakeShareable(new FJsonValueString(DeletedOfflineFile)); });
+	auto ApplyFilenameFilter = [&FilenameFilter](const FString& Filename)
+	{
+		return !FilenameFilter || FilenameFilter(Filename);
+	};
+
+	Algo::TransformIf(Files, FileValues, [&ApplyFilenameFilter](const FSourceControlStateRef& File) { return ApplyFilenameFilter(File->GetFilename()); }, [](const FSourceControlStateRef& File) { return MakeShareable(new FJsonValueString(File->GetFilename())); });
+	Algo::TransformIf(OfflineFiles, FileValues, ApplyFilenameFilter, [](const FString& OfflineFile) { return MakeShareable(new FJsonValueString(OfflineFile)); });
+	Algo::TransformIf(DeletedOfflineFiles, FileValues, ApplyFilenameFilter, [](const FString& DeletedOfflineFile) { return MakeShareable(new FJsonValueString(DeletedOfflineFile)); });
 
 	OutJsonObject->SetArrayField(FILES_NAME, MoveTemp(FileValues));
 }
diff --git a/Engine/Source/Developer/UncontrolledChangelists/Private/UncontrolledChangelistsModule.cpp b/Engine/Source/Developer/UncontrolledChangelists/Private/UncontrolledChangelistsModule.cpp
--- a/Engine/Source/Developer/UncontrolledChangelists/Private/UncontrolledChangelistsModule.cpp
+++ b/Engine/Source/Developer/UncontrolledChangelists/Private/UncontrolledChangelistsModule.cpp
@@ -31,45 +31,97 @@
 
 #define LOCTEXT_NAMESPACE "UncontrolledChangelists"
 
-void FUncontrolledChangelistsModule::FStartupTask::DoWork()
+class FUncontrolledChangelistsDiscoverAssetsTask : public FNonAbandonableTask
+{
+public:
+	FUncontrolledChangelistsDiscoverAssetsTask(FUncontrolledChangelistsModule* InOwner, FARFilter InAssetFilter)
+		: Owner(InOwner)
+		, AssetFilter(MoveTemp(InAssetFilter))
+	{}
+	~FUncontrolledChangelistsDiscoverAssetsTask() {}
+
+	TStatId GetStatId() const
+	{
+		RETURN_QUICK_DECLARE_CYCLE_STAT(FUncontrolledChangelistsDiscoverAssetsTask, STATGROUP_ThreadPoolAsyncTasks);
+	}
+
+	void DoWork();
+	const TSet<FString>& GetAddedAssetsCache() const { return AddedAssetsCache; }
+
+private:
+	FUncontrolledChangelistsModule* Owner;
+	FARFilter AssetFilter;
+	TSet<FString> AddedAssetsCache;
+};
+
+void FUncontrolledChangelistsDiscoverAssetsTask::DoWork()
 {
 	double StartTime = FPlatformTime::Seconds();
-	UE_LOG(LogSourceControl, Log, TEXT("Uncontrolled asset enumeration started..."));
+	UE_LOG(LogSourceControl, Log, TEXT("Uncontrolled asset discovery started..."));
 
-	FAssetRegistryModule& AssetRegistryModule = FModuleManager::GetModuleChecked<FAssetRegistryModule>(TEXT("AssetRegistry"));
-	IAssetRegistry& AssetRegistry = AssetRegistryModule.Get();
+	IAssetRegistry& AssetRegistry = IAssetRegistry::GetChecked();
 
 	TArray<FAssetData> Assets;
-	const bool bIncludeOnlyOnDiskAssets = true;
-	AssetRegistry.GetAllAssets(Assets, bIncludeOnlyOnDiskAssets);
-	for(const FAssetData& AssetData : Assets) 
+	if (AssetFilter.IsEmpty())
+	{
+		const bool bIncludeOnlyOnDiskAssets = true;
+		AssetRegistry.GetAllAssets(Assets, bIncludeOnlyOnDiskAssets);
+	}
+	else
+	{
+		AssetFilter.bIncludeOnlyOnDiskAssets = true;
+		AssetRegistry.GetAssets(AssetFilter, Assets);
+	}
+
+	for (const FAssetData& AssetData : Assets) 
 	{ 
-		if (IsEngineExitRequested())
+		if (Owner->IsStopAssetDiscoveryRequested())
 		{
 			break;
 		}
 		Owner->OnAssetAddedInternal(AssetData, AddedAssetsCache, true);
 	}
 	
-	UE_LOG(LogSourceControl, Log, TEXT("Uncontrolled asset enumeration finished in %s seconds (Found %d uncontrolled assets)"), *FString::SanitizeFloat(FPlatformTime::Seconds() - StartTime), AddedAssetsCache.Num());;
+	UE_LOG(LogSourceControl, Log, TEXT("Uncontrolled asset discovery finished in %s seconds (Found %d uncontrolled assets)"), *FString::SanitizeFloat(FPlatformTime::Seconds() - StartTime), AddedAssetsCache.Num());;
 }
 
+FUncontrolledChangelistsModule::FUncontrolledChangelistsModule() = default;
+FUncontrolledChangelistsModule::~FUncontrolledChangelistsModule() = default;
+
 void FUncontrolledChangelistsModule::StartupModule()
 {
 	bIsEnabled = USourceControlPreferences::AreUncontrolledChangelistsEnabled();
+
+	const bool bIsEnabledThisFrame = IsEnabled();
+	bWasEnabledLastFrame = bIsEnabledThisFrame;
 	
-	if (!IsEnabled())
+	if (bIsEnabledThisFrame)
 	{
-		return;
+		OnEnabled();
 	}
 
+	OnEndFrameDelegateHandle = FCoreDelegates::OnEndFrame.AddRaw(this, &FUncontrolledChangelistsModule::OnEndFrame);
+}
+
+void FUncontrolledChangelistsModule::ShutdownModule()
+{
+	checkf(!DiscoverAssetsTask, TEXT("The discover assets task should be cleaned up when OnEnginePreExit is called at the latest"));
+
+	FCoreDelegates::OnEndFrame.Remove(OnEndFrameDelegateHandle);
+	OnEndFrameDelegateHandle.Reset();
+
+	OnDisabled();
+}
+
+void FUncontrolledChangelistsModule::OnEnabled()
+{
 	// Adds Default Uncontrolled Changelist if it is not already present.
 	GetDefaultUncontrolledChangelistState();
 
 	LoadState();
 
-	OnObjectPreSavedDelegateHandle = FCoreUObjectDelegates::OnObjectPreSave.AddLambda([](UObject* InAsset, const FObjectPreSaveContext& InPreSaveContext) { Get().OnObjectPreSaved(InAsset, InPreSaveContext); });
-	OnEndFrameDelegateHandle = FCoreDelegates::OnEndFrame.AddLambda([]() { Get().OnEndFrame(); });
+	OnObjectPreSavedDelegateHandle = FCoreUObjectDelegates::OnObjectPreSave.AddRaw(this, &FUncontrolledChangelistsModule::OnObjectPreSaved);
+	OnCustomProjectsChangedDelegateHandle = ISourceControlModule::Get().OnCustomProjectsChanged().AddRaw(this, &FUncontrolledChangelistsModule::RequestReloadState);
 
 	// Create initial scan event object
 	InitialScanEvent = MakeShared<FInitialScanEvent>();
@@ -81,70 +133,50 @@ void FUncontrolledChangelistsModule::StartupModule()
 		{
 			return;
 		}
+		InitialScanEvent.Reset();
 
-		StartupTask = MakeUnique<FAsyncTask<FStartupTask>>(this);
-		StartupTask->StartBackgroundTask();
+		StartAssetDiscovery();
+		OnEnginePreExitDelegateHandle = FCoreDelegates::OnEnginePreExit.AddRaw(this, &FUncontrolledChangelistsModule::StopAssetDiscovery);
 
-		InitialScanEvent = nullptr;
-
-		FAssetRegistryModule& AssetRegistryModule = FModuleManager::GetModuleChecked<FAssetRegistryModule>(TEXT("AssetRegistry"));
-		IAssetRegistry& AssetRegistry = AssetRegistryModule.Get();
-		OnAssetAddedDelegateHandle = AssetRegistry.OnAssetAdded().AddLambda([](const struct FAssetData& AssetData) { Get().OnAssetAdded(AssetData); });
-	
-		OnEnginePreExitDelegateHandle = FCoreDelegates::OnEnginePreExit.AddLambda([this]()
-			{
-				if (StartupTask.IsValid())
-				{
-					if (StartupTask->Cancel())
-					{
-						UE_LOG(LogSourceControl, Log, TEXT("UncontrolledChangelistsModule start up task was cancelled due to engine exit"));
-					}
-					else  if (!StartupTask->IsDone())
-					{
-						const double StartTime = FPlatformTime::Seconds();
-
-						UE_LOG(LogSourceControl, Log, TEXT("Waiting on UncontrolledChangelistsModule start up task to complete..."));
-						StartupTask->EnsureCompletion();
-						UE_LOG(LogSourceControl, Log, TEXT("UncontrolledChangelistsModule start up task finished, exit was stalled for %.1f(s)"), FPlatformTime::Seconds() - StartTime);
-					}
-
-					StartupTask.Reset();
-				}
-			});
+		IAssetRegistry& AssetRegistry = IAssetRegistry::GetChecked();
+		OnAssetAddedDelegateHandle = AssetRegistry.OnAssetAdded().AddRaw(this, &FUncontrolledChangelistsModule::OnAssetAdded);
 	}));
-
 }
 
-void FUncontrolledChangelistsModule::ShutdownModule()
+void FUncontrolledChangelistsModule::OnDisabled()
 {
-	// This will make sure callback for initial scan early outs if module was shutdown
-	InitialScanEvent = nullptr;
-
-	checkf(StartupTask == nullptr, TEXT("The startup task should be cleaned up when OnEnginePreExit is called at the latest"));
+	// This will make sure callback for initial scan early outs if feature was disabled
+	InitialScanEvent.Reset();
 
 	if (bIsStateDirty)
 	{
 		SaveState();
-		bIsStateDirty = false;
+		check(!bIsStateDirty); // Should be cleared SaveState
 	}
 
-	FAssetRegistryModule* AssetRegistryModulePtr = static_cast<FAssetRegistryModule*>(FModuleManager::Get().GetModule(TEXT("AssetRegistry")));
-
 	// Check in case AssetRegistry has already been shutdown.
-	if (AssetRegistryModulePtr != nullptr)
+	if (IAssetRegistry* AssetRegistry = IAssetRegistry::Get())
 	{
-		AssetRegistryModulePtr->Get().OnAssetAdded().Remove(OnAssetAddedDelegateHandle);
+		AssetRegistry->OnAssetAdded().Remove(OnAssetAddedDelegateHandle);
 	}
 	OnAssetAddedDelegateHandle.Reset();
 
 	FCoreUObjectDelegates::OnObjectPreSave.Remove(OnObjectPreSavedDelegateHandle);
 	OnObjectPreSavedDelegateHandle.Reset();
 
-	FCoreDelegates::OnEndFrame.Remove(OnEndFrameDelegateHandle);
-	OnEndFrameDelegateHandle.Reset();
+	if (ISourceControlModule* SourceControl = ISourceControlModule::GetPtr())
+	{
+		SourceControl->OnCustomProjectsChanged().Remove(OnCustomProjectsChangedDelegateHandle);
+	}
+	OnCustomProjectsChangedDelegateHandle.Reset();
 
 	FCoreDelegates::OnEnginePreExit.Remove(OnEnginePreExitDelegateHandle);
 	OnEnginePreExitDelegateHandle.Reset();
+
+	bPendingReloadState = false;
+	LoadedCustomProjects.Reset();
+	AddedAssetsCache.Reset();
+	UncontrolledChangelistsStateCache.Reset();
 }
 
 bool FUncontrolledChangelistsModule::IsEnabled() const
@@ -203,6 +235,14 @@ bool FUncontrolledChangelistsModule::AddToUncontrolledChangelist(const TArray<FS
 	FullPaths.Reserve(InFilenames.Num());
 	Algo::Transform(InFilenames, FullPaths, [](const FString& Filename) { return FPaths::ConvertRelativePathToFull(Filename); });
 
+	if (HasCustomProjectFilter())
+	{
+		FullPaths.RemoveAll([this](const FString& Filename)
+		{
+			return !DoesFilePassCustomProjectFilter(Filename);
+		});
+	}
+
 	// Remove from reconcile cache
 	for (const FString& FullPath : FullPaths)
 	{
@@ -240,7 +280,7 @@ FText FUncontrolledChangelistsModule::GetReconcileStatus() const
 		return LOCTEXT("WaitForAssetRegistryStatus", "Waiting for Asset Registry initial scan...");
 	}
 
-	if (StartupTask && !StartupTask->IsDone())
+	if (DiscoverAssetsTask && !DiscoverAssetsTask->IsDone())
 	{
 		return LOCTEXT("ProcessingAssetsStatus", "Processing assets...");
 	}
@@ -255,15 +295,15 @@ bool FUncontrolledChangelistsModule::OnReconcileAssets()
 	const bool bAllowInPIE = false;
 	Scope.MakeDialogDelayed(1.0f, bShowCancelButton, bAllowInPIE);
 
-	if (StartupTask)
+	if (DiscoverAssetsTask)
 	{
-		while (!StartupTask->WaitCompletionWithTimeout(0.016))
+		while (!DiscoverAssetsTask->WaitCompletionWithTimeout(0.016))
 		{
 			Scope.EnterProgressFrame(0.f);
 		}
-				
-		AddedAssetsCache.Append(StartupTask->GetTask().GetAddedAssetsCache());
-		StartupTask = nullptr;
+
+		AddedAssetsCache.Append(DiscoverAssetsTask->GetTask().GetAddedAssetsCache());
+		DiscoverAssetsTask.Reset();
 	}
 
 	if ((!IsEnabled()) || AddedAssetsCache.IsEmpty())
@@ -291,7 +331,7 @@ void FUncontrolledChangelistsModule::OnAssetAdded(const FAssetData& AssetData)
 	OnAssetAddedInternal(AssetData, AddedAssetsCache, false);
 }
 
-void FUncontrolledChangelistsModule::OnAssetAddedInternal(const FAssetData& AssetData, TSet<FString>& InAddedAssetsCache, bool bInStartupTask)
+void FUncontrolledChangelistsModule::OnAssetAddedInternal(const FAssetData& AssetData, TSet<FString>& InAddedAssetsCache, bool bInDiscoveryTask)
 {
 	if (AssetData.HasAnyPackageFlags(PKG_Cooked))
 	{
@@ -304,8 +344,8 @@ void FUncontrolledChangelistsModule::OnAssetAddedInternal(const FAssetData& Asse
 		return;
 	}
 
-	// No need to check for existence when running startup task
-	if (!bInStartupTask)
+	// No need to check for existence when running discovery task
+	if (!bInDiscoveryTask)
 	{
 		if (FPackageName::IsTempPackage(PackagePath.GetPackageName()))
 		{
@@ -332,6 +372,12 @@ void FUncontrolledChangelistsModule::OnAssetAddedInternal(const FAssetData& Asse
 		return;
 	}
 
+	// No need for path check when running discovery task, as it's handled by the ARFilter used by the task
+	if (!bInDiscoveryTask && !DoesFilePassCustomProjectFilter(Fullpath))
+	{
+		return;
+	}
+
 	if (ISourceControlModule::Get().GetProvider().UsesLocalReadOnlyState() && !IFileManager::Get().IsReadOnly(*Fullpath))
 	{
 		InAddedAssetsCache.Add(MoveTemp(Fullpath));
@@ -422,7 +468,7 @@ bool FUncontrolledChangelistsModule::OnRevert(const TArray<FString>& InFilenames
 	return bSuccess;
 }
 
-void FUncontrolledChangelistsModule::OnObjectPreSaved(UObject* InObject, const FObjectPreSaveContext& InPreSaveContext)
+void FUncontrolledChangelistsModule::OnObjectPreSaved(UObject* InObject, FObjectPreSaveContext InPreSaveContext)
 {
 	if (!IsEnabled())
 	{
@@ -698,20 +744,117 @@ void FUncontrolledChangelistsModule::OnStateChanged()
 
 void FUncontrolledChangelistsModule::OnEndFrame()
 {
-	if (StartupTask && StartupTask->IsDone())
+	bool bStateChanged = false;
+
+	if (!IsEngineExitRequested())
 	{
-		AddedAssetsCache.Append(StartupTask->GetTask().GetAddedAssetsCache());
-		StartupTask = nullptr;
+		if (const bool bIsEnabledThisFrame = IsEnabled();
+			bIsEnabledThisFrame != bWasEnabledLastFrame)
+		{
+			if (bWasEnabledLastFrame)
+			{
+				OnDisabled();
+			}
+			else
+			{
+				OnEnabled();
+			}
+			bWasEnabledLastFrame = bIsEnabledThisFrame;
+			bStateChanged = true;
+		}
+
+		if (bPendingReloadState)
+		{
+			ReloadState();
+			check(!bPendingReloadState); // Should be cleared ReloadState
+			bStateChanged = true;
+		}
+	}
+
+	if (DiscoverAssetsTask && DiscoverAssetsTask->IsDone())
+	{
+		AddedAssetsCache.Append(DiscoverAssetsTask->GetTask().GetAddedAssetsCache());
+		DiscoverAssetsTask.Reset();
 	}
 
 	if (bIsStateDirty)
 	{
-		OnUncontrolledChangelistModuleChanged.Broadcast();
+		bStateChanged = true;
 		SaveState();
-		bIsStateDirty = false;
+		check(!bIsStateDirty); // Should be cleared SaveState
+	}
+
+	if (bStateChanged)
+	{
+		OnUncontrolledChangelistModuleChanged.Broadcast();
+	}
+}
+
+void FUncontrolledChangelistsModule::StartAssetDiscovery()
+{
+	checkf(!DiscoverAssetsTask, TEXT("StartAssetDiscovery while another task was still running! Call StopAssetDiscovery first!"));
+	checkf(!InitialScanEvent, TEXT("StartAssetDiscovery called while the asset registry scan was still happening!"));
+
+	FARFilter AssetFilter;
+	if (HasCustomProjectFilter())
+	{
+		checkf(LoadedCustomProjects.Num() > 0, TEXT("HasCustomProjectFilter logic is incompatible with StartAssetDiscovery!"));
+
+		FNameBuilder ProjectContentPackagePath;
+		for (const FSourceControlProjectInfo& Project : LoadedCustomProjects)
+		{
+			for (const FString& ProjectContentDirectory : Project.ContentDirectories)
+			{
+				if (FPackageName::TryConvertFilenameToLongPackageName(ProjectContentDirectory, ProjectContentPackagePath))
+				{
+					AssetFilter.PackagePaths.Add(FName(ProjectContentPackagePath));
+				}
+			}
+		}
+		AssetFilter.bRecursivePaths = true;
+
+		// If AssetFilter.PackagePaths is empty then it means the current set of custom projects haven't mounted their content yet, 
+		// and so there would be nothing to find. We bail here as passing an empty filter would discover everything rather than nothing.
+		if (AssetFilter.PackagePaths.IsEmpty())
+		{
+			UE_LOG(LogSourceControl, Log, TEXT("Uncontrolled asset discovery was skipped due to having no custom project content paths"));
+			return;
+		}
+	}
+
+	DiscoverAssetsTask = MakePimpl<FAsyncTask<FUncontrolledChangelistsDiscoverAssetsTask>>(this, MoveTemp(AssetFilter));
+	DiscoverAssetsTask->StartBackgroundTask();
+}
+
+void FUncontrolledChangelistsModule::StopAssetDiscovery()
+{
+	if (DiscoverAssetsTask.IsValid())
+	{
+		if (DiscoverAssetsTask->Cancel())
+		{
+			UE_LOG(LogSourceControl, Log, TEXT("Uncontrolled asset discovery was cancelled by request"));
+		}
+		else if (!DiscoverAssetsTask->IsDone())
+		{
+			const double StartTime = FPlatformTime::Seconds();
+			bStopAssetDiscoveryRequested = true;
+
+			UE_LOG(LogSourceControl, Log, TEXT("Waiting on uncontrolled asset discovery to stop..."));
+			DiscoverAssetsTask->EnsureCompletion();
+			UE_LOG(LogSourceControl, Log, TEXT("Uncontrolled asset discovery stopped after stalling for %.1f(s)"), FPlatformTime::Seconds() - StartTime);
+
+			bStopAssetDiscoveryRequested = false;
+		}
+
+		DiscoverAssetsTask.Reset();
 	}
 }
 
+bool FUncontrolledChangelistsModule::IsStopAssetDiscoveryRequested() const
+{
+	return bStopAssetDiscoveryRequested;
+}
+
 void FUncontrolledChangelistsModule::CleanAssetsCaches()
 {
 	// Remove files we are already tracking in Uncontrolled Changelists
@@ -752,96 +895,205 @@ FUncontrolledChangelistStateRef FUncontrolledChangelistsModule::GetDefaultUncont
 	return UncontrolledChangelistsStateCache.Emplace(MoveTemp(DefaultUncontrolledChangelist), MakeShared<FUncontrolledChangelistState>(DefaultUncontrolledChangelist, FUncontrolledChangelistState::DEFAULT_UNCONTROLLED_CHANGELIST_DESCRIPTION));
 }
 
+bool FUncontrolledChangelistsModule::HasCustomProjectFilter() const
+{
+	// Note: If these rules change then you'll also need to update the filtering logic in StartAssetDiscovery
+	return LoadedCustomProjects.Num() > 0;
+}
+
+bool FUncontrolledChangelistsModule::DoesFilePassCustomProjectFilter(const FString& InFilename) const
+{
+	if (HasCustomProjectFilter())
+	{
+		for (const FSourceControlProjectInfo& Project : LoadedCustomProjects)
+		{
+			if (DoesFilePassCustomProjectFilter(InFilename, Project))
+			{
+				return true;
+			}
+		}
+		return false;
+	}
+
+	return true;
+}
+
+bool FUncontrolledChangelistsModule::DoesFilePassCustomProjectFilter(const FString& InFilename, const FSourceControlProjectInfo& Project)
+{
+	return FPaths::IsUnderDirectory(InFilename, Project.ProjectDirectory);
+}
+
 void FUncontrolledChangelistsModule::SaveState()
 {
 	SanitizeState();
 
-	TSharedPtr<FJsonObject> RootObject = MakeShareable(new FJsonObject);
-	TArray<TSharedPtr<FJsonValue>> UncontrolledChangelistsArray;
+	auto SaveStateImpl = [this](const FString& PersistentFilePath, const TFunction<bool(const FString&)>& FilenameFilter)
+	{
+		TSharedPtr<FJsonObject> RootObject = MakeShareable(new FJsonObject);
+		RootObject->SetNumberField(VERSION_NAME, VERSION_NUMBER);
 
-	RootObject->SetNumberField(VERSION_NAME, VERSION_NUMBER);
+		TArray<TSharedPtr<FJsonValue>> UncontrolledChangelistsArray;
+		for (const auto& Pair : UncontrolledChangelistsStateCache)
+		{
+			const FUncontrolledChangelist& UncontrolledChangelist = Pair.Key;
+			FUncontrolledChangelistStateRef UncontrolledChangelistState = Pair.Value;
+			TSharedPtr<FJsonObject> UncontrolledChangelistObject = MakeShareable(new FJsonObject);
 
-	for (const auto& Pair : UncontrolledChangelistsStateCache)
-	{
-		const FUncontrolledChangelist& UncontrolledChangelist = Pair.Key;
-		FUncontrolledChangelistStateRef UncontrolledChangelistState = Pair.Value;
-		TSharedPtr<FJsonObject> UncontrolledChangelistObject = MakeShareable(new FJsonObject);
+			UncontrolledChangelist.Serialize(UncontrolledChangelistObject.ToSharedRef());
+			UncontrolledChangelistState->Serialize(UncontrolledChangelistObject.ToSharedRef(), FilenameFilter);
 
-		UncontrolledChangelist.Serialize(UncontrolledChangelistObject.ToSharedRef());
-		UncontrolledChangelistState->Serialize(UncontrolledChangelistObject.ToSharedRef());
+			UncontrolledChangelistsArray.Add(MakeShareable(new FJsonValueObject(UncontrolledChangelistObject)));
+		}
+		RootObject->SetArrayField(CHANGELISTS_NAME, UncontrolledChangelistsArray);
 
-		UncontrolledChangelistsArray.Add(MakeShareable(new FJsonValueObject(UncontrolledChangelistObject)));
-	}
+		using FStringWriter = TJsonWriter<TCHAR, TPrettyJsonPrintPolicy<TCHAR>>;
+		using FStringWriterFactory = TJsonWriterFactory<TCHAR, TPrettyJsonPrintPolicy<TCHAR>>;
 
-	RootObject->SetArrayField(CHANGELISTS_NAME, UncontrolledChangelistsArray);
+		FString RootObjectStr;
+		TSharedRef<FStringWriter> Writer = FStringWriterFactory::Create(&RootObjectStr);
+		FJsonSerializer::Serialize(RootObject.ToSharedRef(), Writer);
 
-	using FStringWriter = TJsonWriter<TCHAR, TPrettyJsonPrintPolicy<TCHAR>>;
-	using FStringWriterFactory = TJsonWriterFactory<TCHAR, TPrettyJsonPrintPolicy<TCHAR>>;
+		if (FFileHelper::SaveStringToFile(RootObjectStr, *PersistentFilePath))
+		{
+			UE_LOG(LogSourceControl, Display, TEXT("Uncontrolled Changelist persistency file saved %s"), *PersistentFilePath);
+		}
+	};
 
-	FString RootObjectStr;
-	TSharedRef<FStringWriter> Writer = FStringWriterFactory::Create(&RootObjectStr);
-	FJsonSerializer::Serialize(RootObject.ToSharedRef(), Writer);
+	if (LoadedCustomProjects.Num() > 0)
+	{
+		// One JSON file per-project
+		for (const FSourceControlProjectInfo& Project : LoadedCustomProjects)
+		{
+			const FString ProjectName = FPaths::GetCleanFilename(Project.ProjectDirectory); // TODO: Add ProjectName to FSourceControlProjectInfo?
+			SaveStateImpl(GetPersistentFilePath(ProjectName), [&Project](const FString& Filename)
+			{
+				return DoesFilePassCustomProjectFilter(Filename, Project);
+			});
+		}
+	}
+	else
+	{
+		// One JSON file for the whole UE project
+		SaveStateImpl(GetPersistentFilePath(FString()), nullptr);
+	}
 
-	FFileHelper::SaveStringToFile(RootObjectStr, *GetPersistentFilePath());
+	bIsStateDirty = false;
 }
 
 void FUncontrolledChangelistsModule::LoadState()
 {
-	FString ImportJsonString;
-	TSharedPtr<FJsonObject> RootObject;
-	uint32 VersionNumber = 0;
-	const TArray<TSharedPtr<FJsonValue>>* UncontrolledChangelistsArray = nullptr;
-
-	if (!FFileHelper::LoadFileToString(ImportJsonString, *GetPersistentFilePath()))
+	auto LoadStateImpl = [this](const FString& PersistentFilePath)
 	{
-		return;
-	}
+		FString ImportJsonString;
+		if (!FFileHelper::LoadFileToString(ImportJsonString, *PersistentFilePath))
+		{
+			return;
+		}
+
+		TSharedRef<TJsonReader<>> JsonReader = TJsonReaderFactory<>::Create(ImportJsonString);
 
-	TSharedRef<TJsonReader<>> JsonReader = TJsonReaderFactory<>::Create(ImportJsonString);
+		TSharedPtr<FJsonObject> RootObject;
+		if (!FJsonSerializer::Deserialize(JsonReader, RootObject))
+		{
+			UE_LOG(LogSourceControl, Error, TEXT("Cannot deserialize RootObject."));
+			return;
+		}
+
+		uint32 VersionNumber = 0;
+		if (!RootObject->TryGetNumberField(VERSION_NAME, VersionNumber))
+		{
+			UE_LOG(LogSourceControl, Error, TEXT("Cannot get field %s."), VERSION_NAME);
+			return;
+		}
 
-	if (!FJsonSerializer::Deserialize(JsonReader, RootObject))
+		if (VersionNumber > VERSION_NUMBER)
+		{
+			UE_LOG(LogSourceControl, Error, TEXT("Version number is invalid (file: %u, current: %u)."), VersionNumber, VERSION_NUMBER);
+			return;
+		}
+
+		const TArray<TSharedPtr<FJsonValue>>* UncontrolledChangelistsArray = nullptr;
+		if (!RootObject->TryGetArrayField(CHANGELISTS_NAME, UncontrolledChangelistsArray))
+		{
+			UE_LOG(LogSourceControl, Error, TEXT("Cannot get field %s."), CHANGELISTS_NAME);
+			return;
+		}
+
+		for (const TSharedPtr<FJsonValue>& JsonValue : *UncontrolledChangelistsArray)
+		{
+			FUncontrolledChangelist TempKey;
+			TSharedRef<FJsonObject> JsonObject = JsonValue->AsObject().ToSharedRef();
+
+			if (!TempKey.Deserialize(JsonObject))
+			{
+				UE_LOG(LogSourceControl, Error, TEXT("Cannot deserialize FUncontrolledChangelist."));
+				continue;
+			}
+
+			FUncontrolledChangelistsStateCache::ValueType& UncontrolledChangelistState = UncontrolledChangelistsStateCache.FindOrAdd(MoveTemp(TempKey), MakeShared<FUncontrolledChangelistState>(TempKey));
+
+			UncontrolledChangelistState->Deserialize(JsonObject);
+		}
+
+		UE_LOG(LogSourceControl, Display, TEXT("Uncontrolled Changelist persistency file loaded %s"), *PersistentFilePath);
+	};
+
+	LoadedCustomProjects = ISourceControlModule::Get().GetCustomProjects();
+	if (LoadedCustomProjects.Num() > 0)
 	{
-		UE_LOG(LogSourceControl, Error, TEXT("Cannot deserialize RootObject."));
-		return;
+		// One JSON file per-project
+		for (const FSourceControlProjectInfo& Project : LoadedCustomProjects)
+		{
+			const FString ProjectName = FPaths::GetCleanFilename(Project.ProjectDirectory); // TODO: Add ProjectName to FSourceControlProjectInfo?
+			LoadStateImpl(GetPersistentFilePath(ProjectName));
+		}
 	}
-
-	if (!RootObject->TryGetNumberField(VERSION_NAME, VersionNumber))
+	else
 	{
-		UE_LOG(LogSourceControl, Error, TEXT("Cannot get field %s."), VERSION_NAME);
-		return;
+		// One JSON file for the whole UE project
+		LoadStateImpl(GetPersistentFilePath(FString()));
 	}
 
-	if (VersionNumber > VERSION_NUMBER)
+	SanitizeState();
+}
+
+void FUncontrolledChangelistsModule::RequestReloadState()
+{
+	bPendingReloadState = true;
+}
+
+void FUncontrolledChangelistsModule::ReloadState()
+{
+	// If the list of projects hasn't actually changed then we can skip this reload
+	if (TArray<FSourceControlProjectInfo> NewCustomProjects = ISourceControlModule::Get().GetCustomProjects();
+		NewCustomProjects == LoadedCustomProjects)
 	{
-		UE_LOG(LogSourceControl, Error, TEXT("Version number is invalid (file: %u, current: %u)."), VersionNumber, VERSION_NUMBER);
+		bPendingReloadState = false;
 		return;
 	}
 
-	if (!RootObject->TryGetArrayField(CHANGELISTS_NAME, UncontrolledChangelistsArray))
+	if (bIsStateDirty)
 	{
-		UE_LOG(LogSourceControl, Error, TEXT("Cannot get field %s."), CHANGELISTS_NAME);
-		return;
+		SaveState();
+		check(!bIsStateDirty); // Should be cleared SaveState
 	}
 
-	for (const TSharedPtr<FJsonValue>& JsonValue : *UncontrolledChangelistsArray)
-	{
-		FUncontrolledChangelist TempKey;
-		TSharedRef<FJsonObject> JsonObject = JsonValue->AsObject().ToSharedRef();
+	// Clear the assets pending reconcile, as we will rebuild that list against the new project roots
+	StopAssetDiscovery();
+	AddedAssetsCache.Reset();
 
-		if (!TempKey.Deserialize(JsonObject))
-		{
-			UE_LOG(LogSourceControl, Error, TEXT("Cannot deserialize FUncontrolledChangelist."));
-			continue;
-		}
+	// Clear any current uncontrolled changelist state, as we will load that from the new project JSON files
+	UncontrolledChangelistsStateCache.Reset();
+	GetDefaultUncontrolledChangelistState();
 
-		FUncontrolledChangelistsStateCache::ValueType& UncontrolledChangelistState = UncontrolledChangelistsStateCache.FindOrAdd(MoveTemp(TempKey), MakeShared<FUncontrolledChangelistState>(TempKey));
+	LoadState();
 
-		UncontrolledChangelistState->Deserialize(JsonObject);
+	if (!InitialScanEvent)
+	{
+		StartAssetDiscovery();
 	}
 
-	SanitizeState();
-
-	UE_LOG(LogSourceControl, Display, TEXT("Uncontrolled Changelist persistency file loaded %s"), *GetPersistentFilePath());
+	bPendingReloadState = false;
 }
 
 void FUncontrolledChangelistsModule::SanitizeState()
@@ -888,37 +1140,10 @@ void FUncontrolledChangelistsModule::SanitizeState()
 	}
 }
 
-FString FUncontrolledChangelistsModule::GetPersistentFilePath() const
-{
-	return FPaths::ProjectSavedDir() + TEXT("SourceControl/UncontrolledChangelists.json");
-}
-
-FString FUncontrolledChangelistsModule::GetUObjectPackageFullpath(const UObject* InObject) const
+FString FUncontrolledChangelistsModule::GetPersistentFilePath(const FString& SubProjectName) const
 {
-	FString Fullpath = TEXT("");
-
-	if (InObject == nullptr)
-	{
-		return Fullpath;
-	}
-
-	const UPackage* Package = InObject->GetPackage();
-
-	if (Package == nullptr)
-	{
-		return Fullpath;
-	}
-
-	const FString LocalFullPath(Package->GetLoadedPath().GetLocalFullPath());
-
-	if (LocalFullPath.IsEmpty())
-	{
-		return Fullpath;
-	}
-
-	Fullpath = FPaths::ConvertRelativePathToFull(LocalFullPath);
-
-	return Fullpath;
+	const FString Filename = SubProjectName.IsEmpty() ? TEXT("UncontrolledChangelists.json") : FString::Printf(TEXT("UncontrolledChangelists_%s.json"), *SubProjectName);
+	return FPaths::Combine(FPaths::ProjectSavedDir(), TEXT("SourceControl"), Filename);
 }
 
 IMPLEMENT_MODULE(FUncontrolledChangelistsModule, UncontrolledChangelists);
diff --git a/Engine/Source/Developer/UncontrolledChangelists/Public/UncontrolledChangelistState.h b/Engine/Source/Developer/UncontrolledChangelists/Public/UncontrolledChangelistState.h
--- a/Engine/Source/Developer/UncontrolledChangelists/Public/UncontrolledChangelistState.h
+++ b/Engine/Source/Developer/UncontrolledChangelists/Public/UncontrolledChangelistState.h
@@ -113,10 +113,11 @@ public:
 	bool ContainsFilename(const FString& PackageFilename) const;
 
 	/**
-	 * Serialize the state of the Uncontrolled Changelist to a Json Object.
+	 * Serialize the state of the Uncontrolled Changelist to a Json Object, optionally filtering the files.
 	 * @param 	OutJsonObject 	The Json object used to serialize.
+	 * @param	FilenameFilter	An optional filter to control while filenames are included in the serialized JSON.
 	 */
-	void Serialize(TSharedRef<class FJsonObject> OutJsonObject) const;
+	void Serialize(TSharedRef<class FJsonObject> OutJsonObject, const TFunction<bool(const FString&)>& FilenameFilter = nullptr) const;
 
 	/**
 	 * Deserialize the state of the Uncontrolled Changelist from a Json Object.
diff --git a/Engine/Source/Developer/UncontrolledChangelists/Public/UncontrolledChangelistsModule.h b/Engine/Source/Developer/UncontrolledChangelists/Public/UncontrolledChangelistsModule.h
--- a/Engine/Source/Developer/UncontrolledChangelists/Public/UncontrolledChangelistsModule.h
+++ b/Engine/Source/Developer/UncontrolledChangelists/Public/UncontrolledChangelistsModule.h
@@ -11,6 +11,10 @@
 #include "Async/AsyncWork.h"
 
 struct FAssetData;
+struct FSourceControlProjectInfo;
+
+class FObjectPreSaveContext;
+class FUncontrolledChangelistsDiscoverAssetsTask;
 
 /**
  * Interface for talking to Uncontrolled Changelists
@@ -29,6 +33,9 @@ public:
 	FOnUncontrolledChangelistModuleChanged OnUncontrolledChangelistModuleChanged;
 
 public:
+	FUncontrolledChangelistsModule();
+	~FUncontrolledChangelistsModule();
+
 	/** IModuleInterface implementation */
 	virtual void StartupModule() override;
 	virtual void ShutdownModule() override;
@@ -119,12 +126,6 @@ public:
 	 */
 	void OnAssetAdded(const FAssetData& AssetData);
 
-	/**
-	 * Delegate callback called when an asset is loaded.
-	 * @param 	InAsset 	The loaded asset.
-	 */
-	void OnAssetLoaded(UObject* InAsset);
-
 	/** Called when "Revert files" button is clicked. Reverts modified files and deletes new ones.
 	 *  @param	InFilenames		The files to be reverted
 	 *	@return true if the provided files were reverted.
@@ -136,7 +137,7 @@ public:
 	 * @param 	InObject 			The saved object.
 	 * @param 	InPreSaveContext 	Interface used to access saved parameters.
 	 */
-	void OnObjectPreSaved(UObject* InObject, const FObjectPreSaveContext& InPreSaveContext);
+	void OnObjectPreSaved(UObject* InObject, FObjectPreSaveContext InPreSaveContext);
 
 	/**
 	 * Moves files to an Uncontrolled Changelist.
@@ -191,19 +192,34 @@ public:
 
 private:
 	/**
-	 * Helper use by Startup task and OnAssetLoaded delegate.
+	 * Helper use by asset discovery task and OnAssetLoaded delegate.
 	 * @param 	InAssetData 		The asset just added.
 	 * @param  	InAddedAssetsCache 	The cache to add the asset to.
-	 * @param	bInStartupTask		If true, this asset was added from the startup task.
+	 * @param	bInDiscoveryTask	If true, this asset was added from the asset discovery task.
 	 * 
 	 */
-	void OnAssetAddedInternal(const FAssetData& InAssetData, TSet<FString>& InAddedAssetsCache, bool bInStartupTask);
+	void OnAssetAddedInternal(const FAssetData& InAssetData, TSet<FString>& InAddedAssetsCache, bool bInDiscoveryTask);
 
 	/**
 	 * Add files to Uncontrolled Changelist. Also adds them to files to reconcile.
 	 */
 	bool AddToUncontrolledChangelist(const TArray<FString>& InFilenames);
 
+	/**
+	 * True if we have custom projects and calling DoesFilePassCustomProjectFilter could ever return something aside from true.
+	 */
+	bool HasCustomProjectFilter() const;
+
+	/**
+	 * Run the given file through the filter for all known custom projects.
+	 */
+	bool DoesFilePassCustomProjectFilter(const FString& InFilename) const;
+
+	/**
+	 * Run the given file through the filter for the given custom project.
+	 */
+	static bool DoesFilePassCustomProjectFilter(const FString& InFilename, const FSourceControlProjectInfo& Project);
+
 	/**
 	 * Saves the state of UncontrolledChangelists to Json for persistency.
 	 */
@@ -214,6 +230,16 @@ private:
 	 */
 	void LoadState();
 	
+	/**
+	 * Request that ReloadState be called at the end of the current frame.
+	 */
+	void RequestReloadState();
+
+	/**
+	 * Reload the state of the UncontrolledChangelists when ISourceControlModule::GetCustomProjects changes.
+	 */
+	void ReloadState();
+
 	/**
 	 * Removes any duplicated files across changelists
 	 */
@@ -225,17 +251,36 @@ private:
 	void OnEndFrame();
 
 	/**
-	 * Helper returning the location of the file used for persistency.
-	 * @return 	A string containing the filepath.
+	 * Called when the uncontrolled changelist feature switches to an enabled state.
+	 */
+	void OnEnabled();
+
+	/**
+	 * Called when the uncontrolled changelist feature switches to a disabled state.
 	 */
-	FString GetPersistentFilePath() const;
+	void OnDisabled();
 
 	/**
-	 * Helper returning the package path where an UObject is located.
-	 * @param 	InObject 	The object used to locate the package.
-	 * @return 	A String containing the filepath of the package.
+	 * Start a new DiscoverAssetsTask for the current state.
+	 * @note InitialScanEvent must be null before calling this, and there shouldn't be any existing DiscoverAssetsTask running.
 	 */
-	FString GetUObjectPackageFullpath(const UObject* InObject) const;
+	void StartAssetDiscovery();
+
+	/**
+	 * Stop the current DiscoverAssetsTask, if any.
+	 */
+	void StopAssetDiscovery();
+
+	/**
+	 * Used by the task to query whether we're waiting for it to finish in StopAssetDiscovery. 
+	 */
+	bool IsStopAssetDiscoveryRequested() const;
+
+	/**
+	 * Helper returning the location of the file used for persistency.
+	 * @return 	A string containing the filepath.
+	 */
+	FString GetPersistentFilePath(const FString& SubProjectName) const;
 
 	/** Called when a state changed either in the module or an Uncontrolled Changelist. */
 	void OnStateChanged();
@@ -255,39 +300,24 @@ private:
 	FUncontrolledChangelistStateRef GetDefaultUncontrolledChangelistState();
 
 private:
-	class FStartupTask : public FNonAbandonableTask
-	{
-	public:
-		
-		FStartupTask(FUncontrolledChangelistsModule* InOwner)
-			: Owner(InOwner)
-		{}
-		~FStartupTask() {}
-		
-		TStatId GetStatId() const
-		{
-			RETURN_QUICK_DECLARE_CYCLE_STAT(FStartupTask, STATGROUP_ThreadPoolAsyncTasks);
-		}
-		
-		void DoWork();
-		const TSet<FString>& GetAddedAssetsCache() const { return AddedAssetsCache; }
-	
-	private:
-		FUncontrolledChangelistsModule* Owner;
-		TSet<FString> AddedAssetsCache;
-	};
+	friend FUncontrolledChangelistsDiscoverAssetsTask;
 
 	// Used to determine if the initial Asset Registry scan was completed or the module was shutdown
 	struct FInitialScanEvent : public TSharedFromThis<FInitialScanEvent> {};
 	TSharedPtr<FInitialScanEvent> InitialScanEvent;
 
-	TUniquePtr<FAsyncTask<FStartupTask>> StartupTask;
+	TPimplPtr<FAsyncTask<FUncontrolledChangelistsDiscoverAssetsTask>> DiscoverAssetsTask;
 	FUncontrolledChangelistsStateCache	UncontrolledChangelistsStateCache;
+	TArray<FSourceControlProjectInfo>	LoadedCustomProjects;
 	TSet<FString>						AddedAssetsCache;
 	FDelegateHandle						OnEnginePreExitDelegateHandle;
 	FDelegateHandle						OnAssetAddedDelegateHandle;
 	FDelegateHandle						OnObjectPreSavedDelegateHandle;
+	FDelegateHandle						OnCustomProjectsChangedDelegateHandle;
 	FDelegateHandle						OnEndFrameDelegateHandle;
+	std::atomic<bool>					bStopAssetDiscoveryRequested = false;
 	bool								bIsEnabled = false;
+	bool								bWasEnabledLastFrame = false;
 	bool								bIsStateDirty = false;
+	bool								bPendingReloadState = false;
 };
