Added API to manage and query uncontrolled changelists via their internal GUID

This is useful for creating programatic uncontrolled changelists (with a fixed ID) that aren't the default changelist

diff --git a/Engine/Source/Developer/UncontrolledChangelists/Private/UncontrolledChangelistsModule.cpp b/Engine/Source/Developer/UncontrolledChangelists/Private/UncontrolledChangelistsModule.cpp
--- a/Engine/Source/Developer/UncontrolledChangelists/Private/UncontrolledChangelistsModule.cpp
+++ b/Engine/Source/Developer/UncontrolledChangelists/Private/UncontrolledChangelistsModule.cpp
@@ -196,6 +196,26 @@ TArray<FUncontrolledChangelistStateRef> FUncontrolledChangelistsModule::GetChang
 	return UncontrolledChangelistStates;
 }
 
+FUncontrolledChangelistStatePtr FUncontrolledChangelistsModule::GetChangelistState(const FUncontrolledChangelist& InUncontrolledChangelist) const
+{
+	if (!IsEnabled())
+	{
+		return nullptr;
+	}
+
+	if (const FUncontrolledChangelistStateRef* UncontrolledChangelistState = UncontrolledChangelistsStateCache.Find(InUncontrolledChangelist))
+	{
+		return *UncontrolledChangelistState;
+	}
+
+	return nullptr;
+}
+
+FUncontrolledChangelistStatePtr FUncontrolledChangelistsModule::GetDefaultChangelistState() const
+{
+	return GetChangelistState(FUncontrolledChangelist(FUncontrolledChangelist::DEFAULT_UNCONTROLLED_CHANGELIST_GUID));
+}
+
 bool FUncontrolledChangelistsModule::OnMakeWritable(const FString& InFilename)
 {
 	if (!IsEnabled())
@@ -660,15 +680,30 @@ void FUncontrolledChangelistsModule::MoveFilesToControlledChangelist(const TArra
 	}
 }
 
-TOptional<FUncontrolledChangelist> FUncontrolledChangelistsModule::CreateUncontrolledChangelist(const FText& InDescription)
+TOptional<FUncontrolledChangelist> FUncontrolledChangelistsModule::CreateUncontrolledChangelist(const FText& InDescription, const TOptional<FUncontrolledChangelist>& InUncontrolledChangelist)
 {
 	if (!IsEnabled())
 	{
 		return TOptional<FUncontrolledChangelist>();
 	}
 
+	if (InUncontrolledChangelist)
+	{
+		if (InUncontrolledChangelist->IsDefault())
+		{
+			UE_LOG(LogSourceControl, Error, TEXT("Cannot use CreateUncontrolledChangelist with the Default Uncontrolled Changelist."));
+			return TOptional<FUncontrolledChangelist>();
+		}
+
+		if (FUncontrolledChangelistStateRef* UncontrolledChangelistState = UncontrolledChangelistsStateCache.Find(*InUncontrolledChangelist))
+		{
+			(*UncontrolledChangelistState)->SetDescription(InDescription);
+			return *InUncontrolledChangelist;
+		}
+	}
+	
 	// Default constructor will generate a new GUID.
-	FUncontrolledChangelist NewUncontrolledChangelist;
+	FUncontrolledChangelist NewUncontrolledChangelist = InUncontrolledChangelist ? *InUncontrolledChangelist : FUncontrolledChangelist();
 	UncontrolledChangelistsStateCache.Emplace(NewUncontrolledChangelist, MakeShared<FUncontrolledChangelistState>(NewUncontrolledChangelist, InDescription));
 
 	OnStateChanged();
@@ -892,7 +927,7 @@ FUncontrolledChangelistStateRef FUncontrolledChangelistsModule::GetDefaultUncont
 		return *DefaultUncontrolledChangelistState;
 	}
 
-	return UncontrolledChangelistsStateCache.Emplace(MoveTemp(DefaultUncontrolledChangelist), MakeShared<FUncontrolledChangelistState>(DefaultUncontrolledChangelist, FUncontrolledChangelistState::DEFAULT_UNCONTROLLED_CHANGELIST_DESCRIPTION));
+	return UncontrolledChangelistsStateCache.Emplace(DefaultUncontrolledChangelist, MakeShared<FUncontrolledChangelistState>(DefaultUncontrolledChangelist, FUncontrolledChangelistState::DEFAULT_UNCONTROLLED_CHANGELIST_DESCRIPTION));
 }
 
 bool FUncontrolledChangelistsModule::HasCustomProjectFilter() const
diff --git a/Engine/Source/Developer/UncontrolledChangelists/Public/UncontrolledChangelistsModule.h b/Engine/Source/Developer/UncontrolledChangelists/Public/UncontrolledChangelistsModule.h
--- a/Engine/Source/Developer/UncontrolledChangelists/Public/UncontrolledChangelistsModule.h
+++ b/Engine/Source/Developer/UncontrolledChangelists/Public/UncontrolledChangelistsModule.h
@@ -48,7 +48,17 @@ public:
 	/**
 	 * Get the changelist state of each cached Uncontrolled Changelist.
 	 */
-	 TArray<FUncontrolledChangelistStateRef> GetChangelistStates() const;
+	TArray<FUncontrolledChangelistStateRef> GetChangelistStates() const;
+
+	/**
+	 * Get the changelist state of the given Uncontrolled Changelist.
+	 */
+	FUncontrolledChangelistStatePtr GetChangelistState(const FUncontrolledChangelist& InUncontrolledChangelist) const;
+
+	/**
+	 * Get the changelist state of the default Uncontrolled Changelist.
+	 */
+	FUncontrolledChangelistStatePtr GetDefaultChangelistState() const;
 
 	/**
 	 * Called when file has been made writable. Adds the file to the reconcile list because we don't know yet if it will change.
@@ -172,11 +182,12 @@ public:
 
 	/**
 	 * Creates a new Uncontrolled Changelist.
-	 * @param	InDescription	The description of the newly created Uncontrolled Changelist.
+	 * @param	InDescription				The description of the newly created Uncontrolled Changelist.
+	 * @param	InUncontrolledChangelist	An optional Uncontrolled Changelist to create (or find) via its ID. Should not be the default Uncontrolled Changelist.
 	 * return	TOptional<FUncontrolledChangelist> set with the new Uncontrolled Changelist key if succeeded.
 	 */
-	TOptional<FUncontrolledChangelist> CreateUncontrolledChangelist(const FText& InDescription);
-	
+	TOptional<FUncontrolledChangelist> CreateUncontrolledChangelist(const FText& InDescription, const TOptional<FUncontrolledChangelist>& InUncontrolledChangelist = TOptional<FUncontrolledChangelist>());
+
 	/**
      * Edits an Uncontrolled Changelist's description
 	 * @param	InUncontrolledChangelist	The Uncontrolled Changelist to modify. Should not be the default Uncontrolled Changelist.
