SourceControl: Fixed `checkf` crash in uncontrolled changelist module when running an uncooked build.

The asset discovery scan is (near) instant in cooked builds but when running uncooked one could hit the `checkf` at the beginning of StartAssetDiscovery().

#rnx

diff --git a/Engine/Source/Developer/UncontrolledChangelists/Private/UncontrolledChangelistsModule.cpp b/Engine/Source/Developer/UncontrolledChangelists/Private/UncontrolledChangelistsModule.cpp
--- a/Engine/Source/Developer/UncontrolledChangelists/Private/UncontrolledChangelistsModule.cpp
+++ b/Engine/Source/Developer/UncontrolledChangelists/Private/UncontrolledChangelistsModule.cpp
@@ -147,6 +147,8 @@ void FUncontrolledChangelistsModule::OnDisabled()
 {
 	// This will make sure callback for initial scan early outs if feature was disabled
 	InitialScanEvent.Reset();
+	
+	StopAssetDiscovery();
 
 	if (bIsStateDirty)
 	{
