PR #13035: Fix bugs with "Save" and "Save To Changelist" when the destination is an Uncontrolled Changelist.

- When adding a list of files to Uncontrolled Changelists, group them first by their existing Uncontrolled Changelist, if any, to avoid duplication and only add new files to the default Uncontrolled Changelist.
- When moving files to an Uncontrolled Changelist, remove them first from any existing Uncontrolled Changelist to avoid duplication.
- Disable the "Save" and "Save To Changelist" context menu entries when they would otherwise be noops.
- In "Save To Changelist" make the call to TryMoveFiles() after the call to SavePackages.  This fixes a bug where the file could end up in both the default controlled changelist and the chosen uncontrolled changelist due to order of operations.
- Request a changelists refresh after both "Save" and "Save To Changelist" to avoid the need for a manual refresh in some cases after these operations.

diff --git a/Engine/Source/Developer/UncontrolledChangelists/Private/UncontrolledChangelistState.cpp b/Engine/Source/Developer/UncontrolledChangelists/Private/UncontrolledChangelistState.cpp
--- a/Engine/Source/Developer/UncontrolledChangelists/Private/UncontrolledChangelistState.cpp
+++ b/Engine/Source/Developer/UncontrolledChangelists/Private/UncontrolledChangelistState.cpp
@@ -93,7 +93,19 @@ TArray<FString> FUncontrolledChangelistState::GetFilenames() const
 
 bool FUncontrolledChangelistState::ContainsFilename(const FString& PackageFilename) const
 {
-	return  OfflineFiles.Contains(PackageFilename) || Algo::AnyOf(Files, [&PackageFilename](const TSharedRef<ISourceControlState>& FileState){ return FileState->GetFilename() == PackageFilename; });
+	if (OfflineFiles.Contains(PackageFilename))
+	{
+		return true;
+	}
+
+	ISourceControlProvider& SourceControlProvider = ISourceControlModule::Get().GetProvider();
+
+	if (FSourceControlStatePtr FileState = SourceControlProvider.GetState(PackageFilename, EStateCacheUsage::Use))
+	{
+		return Files.Contains(FileState.ToSharedRef());
+	}
+
+	return  false;
 }
 
 void FUncontrolledChangelistState::Serialize(TSharedRef<FJsonObject> OutJsonObject, const TFunction<bool(const FString&)>& FilenameFilter) const
@@ -231,6 +243,25 @@ bool FUncontrolledChangelistState::RemoveFiles(const TArray<FSourceControlStateR
 	return bOutChanged;
 }
 
+bool FUncontrolledChangelistState::RemoveFiles(const TArray<FString>& InFilenames)
+{
+	bool bOutChanged = false;
+
+	ISourceControlProvider& SourceControlProvider = ISourceControlModule::Get().GetProvider();
+
+	for (const FString& Filename : SourceControlHelpers::AbsoluteFilenames(InFilenames))
+	{
+		bOutChanged |= (OfflineFiles.Remove(Filename) > 0);
+		
+		if (FSourceControlStatePtr FileState = SourceControlProvider.GetState(Filename, EStateCacheUsage::Use))
+		{
+			bOutChanged |= (Files.Remove(FileState.ToSharedRef()) > 0);
+		}
+	}
+
+	return bOutChanged;
+}
+
 bool FUncontrolledChangelistState::UpdateStatus()
 {
 	TArray<FString> FilesToUpdate;
diff --git a/Engine/Source/Developer/UncontrolledChangelists/Private/UncontrolledChangelistsModule.cpp b/Engine/Source/Developer/UncontrolledChangelists/Private/UncontrolledChangelistsModule.cpp
--- a/Engine/Source/Developer/UncontrolledChangelists/Private/UncontrolledChangelistsModule.cpp
+++ b/Engine/Source/Developer/UncontrolledChangelists/Private/UncontrolledChangelistsModule.cpp
@@ -276,7 +276,95 @@ bool FUncontrolledChangelistsModule::AddToUncontrolledChangelist(const TArray<FS
 		AddedAssetsCache.Remove(FullPath);
 	}
 
-	return AddFilesToDefaultUncontrolledChangelist(FullPaths, FUncontrolledChangelistState::ECheckFlags::NotCheckedOut);
+	// Group files by their UncontrolledChangelist and then add each group to it's corresponding UncontrolledChangelist.
+	TMap<FUncontrolledChangelist, TArray<FString>> UncontrolledChangelistToFullPaths;
+	GroupFilesByUncontrolledChangelist(MoveTemp(FullPaths), UncontrolledChangelistToFullPaths);
+
+	bool bAreAllFilesAdded = true;
+
+	for (const TTuple<FUncontrolledChangelist, TArray<FString>>& Pair : UncontrolledChangelistToFullPaths)
+	{
+		const FUncontrolledChangelist& UncontrolledChangelist = Pair.Key;
+		const TArray<FString>& UncontrolledChangelistFilenames = Pair.Value;
+
+		bAreAllFilesAdded &= AddFilesToUncontrolledChangelist(UncontrolledChangelist, UncontrolledChangelistFilenames, FUncontrolledChangelistState::ECheckFlags::NotCheckedOut);
+	}
+
+	return bAreAllFilesAdded;
+}
+
+bool FUncontrolledChangelistsModule::RemoveFromUncontrolledChangelist(const TArray<FString>& InFilenames)
+{
+	if (!IsEnabled())
+	{
+		return false;
+	}
+
+	TRACE_CPUPROFILER_EVENT_SCOPE(FUncontrolledChangelistsModule::RemoveFromUncontrolledChangelist);
+
+	TArray<FString> FullPaths;
+	FullPaths.Reserve(InFilenames.Num());
+	Algo::Transform(InFilenames, FullPaths, [](const FString& Filename) { return FPaths::ConvertRelativePathToFull(Filename); });
+
+	if (HasCustomProjectFilter())
+	{
+		FullPaths.RemoveAll([this](const FString& Filename)
+		{
+			return !DoesFilePassCustomProjectFilter(Filename);
+		});
+	}
+
+	// Remove from reconcile cache
+	for (const FString& FullPath : FullPaths)
+	{
+		AddedAssetsCache.Remove(FullPath);
+	}
+
+	// Group files by their UncontrolledChangelist and then remove each group from it's corresponding UncontrolledChangelist.
+	TMap<FUncontrolledChangelist, TArray<FString>> UncontrolledChangelistToFullPaths;
+	GroupFilesByUncontrolledChangelist(MoveTemp(FullPaths), UncontrolledChangelistToFullPaths);
+
+	bool bAreAnyFilesRemoved = false;
+
+	for (const TTuple<FUncontrolledChangelist, TArray<FString>>& Pair : UncontrolledChangelistToFullPaths)
+	{
+		const FUncontrolledChangelist& UncontrolledChangelist = Pair.Key;
+		const TArray<FString>& UncontrolledChangelistFilenames = Pair.Value;
+
+		bAreAnyFilesRemoved |= RemoveFilesFromUncontrolledChangelist(UncontrolledChangelist, UncontrolledChangelistFilenames);
+	}
+
+	return bAreAnyFilesRemoved;
+}
+
+void FUncontrolledChangelistsModule::GroupFilesByUncontrolledChangelist(TArray<FString>&& InFilenames, TMap<FUncontrolledChangelist, TArray<FString>>& OutUncontrolledChangelistToFilenames) const
+{
+	for (const TTuple<FUncontrolledChangelist, TSharedRef<FUncontrolledChangelistState>>& Pair : UncontrolledChangelistsStateCache)
+	{
+		const FUncontrolledChangelist& UncontrolledChangelist = Pair.Key;
+		const FUncontrolledChangelistStateRef& UncontrolledChangelistState = Pair.Value;
+
+		TArray<FString>& UncontrolledChangelistFilenames = OutUncontrolledChangelistToFilenames.FindOrAdd(UncontrolledChangelist);
+
+		for (const FString& Filename : InFilenames)
+		{
+			if (UncontrolledChangelistState->ContainsFilename(Filename))
+			{
+				UncontrolledChangelistFilenames.Add(Filename);
+			}
+		}
+
+		for (const FString& Filename : UncontrolledChangelistFilenames)
+		{
+			InFilenames.Remove(Filename);
+		}
+	}
+
+	TArray<FString>& DefaultUncontrolledChangelistFilenames = OutUncontrolledChangelistToFilenames.FindOrAdd(FUncontrolledChangelist{ FUncontrolledChangelist::DEFAULT_UNCONTROLLED_CHANGELIST_GUID });
+	for (const FString& Filename : InFilenames)
+	{
+		DefaultUncontrolledChangelistFilenames.Add(Filename);
+	}
 }
 
 void FUncontrolledChangelistsModule::UpdateStatus()
@@ -341,7 +429,10 @@ bool FUncontrolledChangelistsModule::OnReconcileAssets()
 	Scope.EnterProgressFrame(0.f, LOCTEXT("ReconcileAssetsProgress", "Reconciling assets"));
 
 	CleanAssetsCaches();
-	bool bHasStateChanged = AddFilesToDefaultUncontrolledChangelist(AddedAssetsCache.Array(), FUncontrolledChangelistState::ECheckFlags::All);
+
+	bool bHasStateChanged = AddFilesToUncontrolledChangelist(
+		FUncontrolledChangelist{ FUncontrolledChangelist::DEFAULT_UNCONTROLLED_CHANGELIST_GUID },
+		AddedAssetsCache.Array(), FUncontrolledChangelistState::ECheckFlags::All);
 
 	AddedAssetsCache.Empty();
 
@@ -596,8 +687,11 @@ void FUncontrolledChangelistsModule::MoveFilesToUncontrolledChangelist(const TAr
 	RevertOperation->SetSoftRevert(true);
 	SourceControlProvider.Execute(RevertOperation, Filenames);
 
+	// Remove files from any existing UncontrolledChangelist
+	bHasStateChanged = RemoveFromUncontrolledChangelist(Filenames);
+
 	// Add all files to their UncontrolledChangelist
-	bHasStateChanged = (*ChangelistState)->AddFiles(Filenames, FUncontrolledChangelistState::ECheckFlags::None);
+	bHasStateChanged |= (*ChangelistState)->AddFiles(Filenames, FUncontrolledChangelistState::ECheckFlags::None);
 
 	if (bHasStateChanged)
 	{
@@ -910,14 +1004,32 @@ void FUncontrolledChangelistsModule::CleanAssetsCaches()
 	}
 }
 
-bool FUncontrolledChangelistsModule::AddFilesToDefaultUncontrolledChangelist(const TArray<FString>& InFilenames, const FUncontrolledChangelistState::ECheckFlags InCheckFlags)
+bool FUncontrolledChangelistsModule::AddFilesToUncontrolledChangelist(const FUncontrolledChangelist& InUncontrolledChangelist, const TArray<FString>& InFilenames, const FUncontrolledChangelistState::ECheckFlags InCheckFlags)
 {
 	bool bHasStateChanged = false;
 
-	FUncontrolledChangelistStateRef UncontrolledChangelistState = GetDefaultUncontrolledChangelistState();
+	if (FUncontrolledChangelistStatePtr UncontrolledChangelistState = GetUncontrolledChangelistState(InUncontrolledChangelist))
+	{
+		// Try to add files, they will be added only if they pass the required checks
+		bHasStateChanged = UncontrolledChangelistState->AddFiles(InFilenames, InCheckFlags);
+	}
+
+	if (bHasStateChanged)
+	{
+		OnStateChanged();
+	}
+
+	return bHasStateChanged;
+}
 
-	// Try to add files, they will be added only if they pass the required checks
-	bHasStateChanged = UncontrolledChangelistState->AddFiles(InFilenames, InCheckFlags);
+bool FUncontrolledChangelistsModule::RemoveFilesFromUncontrolledChangelist(const FUncontrolledChangelist& InUncontrolledChangelist, const TArray<FString>& InFilenames)
+{
+	bool bHasStateChanged = false;
+
+	if (FUncontrolledChangelistStatePtr UncontrolledChangelistState = GetUncontrolledChangelistState(InUncontrolledChangelist))
+	{
+		bHasStateChanged = UncontrolledChangelistState->RemoveFiles(InFilenames);
+	}
 
 	if (bHasStateChanged)
 	{
@@ -929,15 +1041,24 @@ bool FUncontrolledChangelistsModule::AddFilesToDefaultUncontrolledChangelist(con
 
 FUncontrolledChangelistStateRef FUncontrolledChangelistsModule::GetDefaultUncontrolledChangelistState()
 {
-	FUncontrolledChangelist DefaultUncontrolledChangelist(FUncontrolledChangelist::DEFAULT_UNCONTROLLED_CHANGELIST_GUID);
-	FUncontrolledChangelistStateRef* DefaultUncontrolledChangelistState = UncontrolledChangelistsStateCache.Find(DefaultUncontrolledChangelist);
+	return GetUncontrolledChangelistState(FUncontrolledChangelist{ FUncontrolledChangelist::DEFAULT_UNCONTROLLED_CHANGELIST_GUID }).ToSharedRef();
+}
+
+FUncontrolledChangelistStatePtr FUncontrolledChangelistsModule::GetUncontrolledChangelistState(const FUncontrolledChangelist& InUncontrolledChangelist)
+{
+	FUncontrolledChangelistStateRef* UncontrolledChangelistState = UncontrolledChangelistsStateCache.Find(InUncontrolledChangelist);
+
+	if (UncontrolledChangelistState != nullptr)
+	{
+		return *UncontrolledChangelistState;
+	}
 
-	if (DefaultUncontrolledChangelistState != nullptr)
+	if (InUncontrolledChangelist.IsDefault())
 	{
-		return *DefaultUncontrolledChangelistState;
+		return UncontrolledChangelistsStateCache.Emplace(InUncontrolledChangelist, MakeShared<FUncontrolledChangelistState>(InUncontrolledChangelist, FUncontrolledChangelistState::DEFAULT_UNCONTROLLED_CHANGELIST_DESCRIPTION));
 	}
 
-	return UncontrolledChangelistsStateCache.Emplace(DefaultUncontrolledChangelist, MakeShared<FUncontrolledChangelistState>(DefaultUncontrolledChangelist, FUncontrolledChangelistState::DEFAULT_UNCONTROLLED_CHANGELIST_DESCRIPTION));
+	return nullptr;
 }
 
 bool FUncontrolledChangelistsModule::HasCustomProjectFilter() const
diff --git a/Engine/Source/Developer/UncontrolledChangelists/Public/UncontrolledChangelistState.h b/Engine/Source/Developer/UncontrolledChangelists/Public/UncontrolledChangelistState.h
--- a/Engine/Source/Developer/UncontrolledChangelists/Public/UncontrolledChangelistState.h
+++ b/Engine/Source/Developer/UncontrolledChangelists/Public/UncontrolledChangelistState.h
@@ -141,6 +141,13 @@ public:
 	 */
 	bool RemoveFiles(const TArray<FSourceControlStateRef>& InFileStates);
 
+	/**
+	 * Removes the given filenames from the file states or offline files.
+	 * @param 	InFilenames 	The filenames to be removed.
+	 * @return 	True if a change has been performed in the Uncontrolled Changelist State.
+	 */
+	bool RemoveFiles(const TArray<FString>& InFilenames);
+
 	/**
 	 * Updates the status of all files contained in this changelist.
 	 * @return 	True if the state has been modified.
diff --git a/Engine/Source/Developer/UncontrolledChangelists/Public/UncontrolledChangelistsModule.h b/Engine/Source/Developer/UncontrolledChangelists/Public/UncontrolledChangelistsModule.h
--- a/Engine/Source/Developer/UncontrolledChangelists/Public/UncontrolledChangelistsModule.h
+++ b/Engine/Source/Developer/UncontrolledChangelists/Public/UncontrolledChangelistsModule.h
@@ -230,6 +230,18 @@ private:
 	 */
 	bool AddToUncontrolledChangelist(const TArray<FString>& InFilenames);
 
+	/**
+	 * Removes the given files from their associated Uncontrolled Changelist, if any.  Also removes them from files to reconcile.
+	 * @return true if any of the given files are removed.
+	 */
+	bool RemoveFromUncontrolledChangelist(const TArray<FString>& InFilenames);
+
+	/**
+	 * Groups the given files by their associated Uncontrolled Changelist and returns the groupings in the given map.
+	 * Any of the files that are not already associated with an Uncontrolled Changelist will be added to the default Uncontrolled Changelist's group.
+	 */
+	void GroupFilesByUncontrolledChangelist(TArray<FString>&& InFilenames, TMap<FUncontrolledChangelist, TArray<FString>>& OutUncontrolledChangelistToFilenames) const;
+
 	/**
 	 * True if we have custom projects and calling DoesFilePassCustomProjectFilter could ever return something aside from true.
 	 */
@@ -314,16 +326,32 @@ private:
 	void CleanAssetsCaches();
 
 	/**
-	 * Try to add the provided filenames to the default Uncontrolled Changelist.
-	 * @param 	InFilenames 	The files to add.
-	 * @param 	InCheckFlags 	The required checks to check the file against before adding.
-	 * @return 	True file have been added.
+	 * Try to add the provided filenames to the given Uncontrolled Changelist.
+	 * @param	InUncontrolledChangelist	The Uncontrolled Changelist to add to.
+	 * @param	InFilenames 				The files to add.
+	 * @param	InCheckFlags 				The required checks to check the file against before adding.
+	 * @return	True files have been added.
 	 */
-	bool AddFilesToDefaultUncontrolledChangelist(const TArray<FString>& InFilenames, const FUncontrolledChangelistState::ECheckFlags InCheckFlags);
+	bool AddFilesToUncontrolledChangelist(const FUncontrolledChangelist& InUncontrolledChangelist, const TArray<FString>& InFilenames, const FUncontrolledChangelistState::ECheckFlags InCheckFlags);
+
+	/**
+	 * Try to remove the provided filenames from the given Uncontrolled Changelist.
+	 * @param	InUncontrolledChangelist	The Uncontrolled Changelist to remove from.
+	 * @param	InFilenames 				The files to remove.
+	 * @return	True files have been removed.
+	 */
+	bool RemoveFilesFromUncontrolledChangelist(const FUncontrolledChangelist& InUncontrolledChangelist, const TArray<FString>& InFilenames);
 
 	/** Returns the default Uncontrolled Changelist state, creates it if it does not exist. */
 	FUncontrolledChangelistStateRef GetDefaultUncontrolledChangelistState();
 
+	/**
+	 * Returns the given Uncontrolled Changelist state if it exists.
+	 * If the given Uncontrolled Changelist is the default one and it does not exist, then it will be created and returned.
+	 * Otherwise, if not the default Uncontrolled Changelist, then nullptr will be returned if it does not exist.
+	 */
+	FUncontrolledChangelistStatePtr GetUncontrolledChangelistState(const FUncontrolledChangelist& InUncontrolledChangelist);
+
 private:
 	friend FUncontrolledChangelistsDiscoverAssetsTask;
 
diff --git a/Engine/Source/Editor/SourceControlWindows/Private/SSourceControlChangelists.cpp b/Engine/Source/Editor/SourceControlWindows/Private/SSourceControlChangelists.cpp
--- a/Engine/Source/Editor/SourceControlWindows/Private/SSourceControlChangelists.cpp
+++ b/Engine/Source/Editor/SourceControlWindows/Private/SSourceControlChangelists.cpp
@@ -1773,7 +1773,7 @@ void SSourceControlChangelistsWidget::GetSelectedFiles(TArray<FString>& OutContr
 
 void SSourceControlChangelistsWidget::GetSelectedFileStates(TArray<FSourceControlStateRef>& OutControlledFileStates, TArray<FSourceControlStateRef>& OutUncontrolledFileStates)
 {
-	TArray<TSharedPtr<IChangelistTreeItem>> SelectedItems = FileListView->GetSelectedItems();
+	TArray<TSharedPtr<IChangelistTreeItem>> SelectedItems = GetActiveFileListView().GetSelectedItems();
 
 	for (const TSharedPtr<IChangelistTreeItem>& Item : SelectedItems)
 	{
@@ -2803,6 +2803,11 @@ TValueOrError<void, void> SSourceControlChangelistsWidget::TryMoveFiles()
 			{
 				UncontrolledChangelistsModule.MoveFilesToUncontrolledChangelist(SelectedControlledFileStates, SelectedUnControlledFileStates, NewUncontrolledChangelist.GetValue());
 			}
+			// SelectedControlledFiles can contain UnsavedAssets when offline even if both sets of states are empty.
+			else if (!SelectedControlledFiles.IsEmpty())
+			{
+				UncontrolledChangelistsModule.MoveFilesToUncontrolledChangelist(SelectedControlledFiles, NewUncontrolledChangelist.GetValue());
+			}
 		});
 	}
 	else // Move files to an existing changelist or uncontrolled changelist.
@@ -2864,13 +2869,18 @@ TValueOrError<void, void> SSourceControlChangelistsWidget::TryMoveFiles()
 
 			GetSelectedFileStates(SelectedControlledFileStates, SelectedUnControlledFileStates);
 
-			if ((!SelectedControlledFileStates.IsEmpty()) || (!SelectedUnControlledFileStates.IsEmpty()))
+			ExecuteUncontrolledChangelistOperation(LOCTEXT("Moving_Uncontrolled_Changelist_To", "Moving uncontrolled files..."), [&]()
 			{
-				ExecuteUncontrolledChangelistOperation(LOCTEXT("Moving_Uncontrolled_Changelist_To", "Moving uncontrolled files..."), [&]()
+				if ((!SelectedControlledFileStates.IsEmpty()) || (!SelectedUnControlledFileStates.IsEmpty()))
 				{
 					FUncontrolledChangelistsModule::Get().MoveFilesToUncontrolledChangelist(SelectedControlledFileStates, SelectedUnControlledFileStates, UncontrolledChangelist);
-				});
-			}
+				}
+				// SelectedControlledFiles can contain UnsavedAssets when offline even if both sets of states are empty.
+				else if (!SelectedControlledFiles.IsEmpty())
+				{
+					FUncontrolledChangelistsModule::Get().MoveFilesToUncontrolledChangelist(SelectedControlledFiles, UncontrolledChangelist);
+				}
+			});
 		}
 	}
 	if (bFailed)
@@ -2986,6 +2996,13 @@ TSharedPtr<SWidget> SSourceControlChangelistsWidget::OnOpenContextMenu()
 				}
 				const bool bOnlyDirty = false;
 				UEditorLoadingAndSavingUtils::SavePackages(Packages, bOnlyDirty);
+
+				// Request a changelists refresh in the next tick
+				bShouldRefresh = true;
+
+			}), FCanExecuteAction::CreateLambda([this]
+			{
+				return HasFilesSelected();
 			})));
 
 		Section.AddMenuEntry(
@@ -2998,6 +3015,7 @@ TSharedPtr<SWidget> SSourceControlChangelistsWidget::OnOpenContextMenu()
 				
 				TArray<FString> SelectedFiles;
 				GetSelectedFiles(SelectedFiles, SelectedFiles);
+				
 				TArray<FString> PackageNames;
 				for (const FString& Filename : SelectedFiles)
 				{
@@ -3009,12 +3027,6 @@ TSharedPtr<SWidget> SSourceControlChangelistsWidget::OnOpenContextMenu()
 					return;
 				}
 				
-				// Handles almost everything else
-				if (TryMoveFiles().HasError())
-				{
-					return;
-				}
-
 				// Actually save the assets
 				TArray<UPackage*> Packages;
 				Packages.Reserve(PackageNames.Num());
@@ -3034,6 +3046,17 @@ TSharedPtr<SWidget> SSourceControlChangelistsWidget::OnOpenContextMenu()
 				}
 				const bool bOnlyDirty = false;
 				UEditorLoadingAndSavingUtils::SavePackages(Packages, bOnlyDirty);
+
+				// Move the files only after saving the packages.
+				// If we are moving files to an uncontrolled changelist, they will be reverted from source control during the move.
+				TryMoveFiles();
+				
+				// Request a changelists refresh in the next tick
+				bShouldRefresh = true;
+
+			}), FCanExecuteAction::CreateLambda([this]
+			{
+				return HasFilesSelected();
 			})));
 
 		Section.AddSeparator("ActionSeparator");
