Updated headers using UnrealCodeFixup to make sure dllstorage is on methods/staticvars instead of types. Part 4/n

diff --git a/Engine/Source/Developer/UncontrolledChangelists/Public/UncontrolledChangelistState.h b/Engine/Source/Developer/UncontrolledChangelists/Public/UncontrolledChangelistState.h
--- a/Engine/Source/Developer/UncontrolledChangelists/Public/UncontrolledChangelistState.h
+++ b/Engine/Source/Developer/UncontrolledChangelists/Public/UncontrolledChangelistState.h
@@ -18,15 +18,17 @@
 #include "UObject/NameTypes.h"
 #include "UncontrolledChangelist.h"
 
+#define UE_API UNCONTROLLEDCHANGELISTS_API
+
 class FJsonObject;
 
-class UNCONTROLLEDCHANGELISTS_API FUncontrolledChangelistState : public TSharedFromThis<FUncontrolledChangelistState, ESPMode::ThreadSafe>
+class FUncontrolledChangelistState : public TSharedFromThis<FUncontrolledChangelistState, ESPMode::ThreadSafe>
 {
 public:
 	static constexpr const TCHAR* FILES_NAME = TEXT("files");
 	static constexpr const TCHAR* NAME_NAME = TEXT("name");
 	static constexpr const TCHAR* DESCRIPTION_NAME = TEXT("description");
-	static const FText DEFAULT_UNCONTROLLED_CHANGELIST_DESCRIPTION;
+	static UE_API const FText DEFAULT_UNCONTROLLED_CHANGELIST_DESCRIPTION;
 
 	enum class ECheckFlags
 	{
@@ -47,9 +49,9 @@ public:
 	// An FUncontrolledChangelistState should always reference a given Changelist (with a valid GUID).
 	FUncontrolledChangelistState() = delete;
 
-	FUncontrolledChangelistState(const FUncontrolledChangelist& InUncontrolledChangelist);
+	UE_API FUncontrolledChangelistState(const FUncontrolledChangelist& InUncontrolledChangelist);
 	
-	FUncontrolledChangelistState(const FUncontrolledChangelist& InUncontrolledChangelist, const FText& InDescription);
+	UE_API FUncontrolledChangelistState(const FUncontrolledChangelist& InUncontrolledChangelist, const FText& InDescription);
 
 	// Uncontrolled Changelist states are unique and non-copyable, should always be used by reference to preserve cache coherence.
 	FUncontrolledChangelistState(const FUncontrolledChangelistState& InUncontrolledChangelistState) = delete;
@@ -59,72 +61,72 @@ public:
 	 * Get the name of the icon graphic we should use to display the state in a UI.
 	 * @returns the name of the icon to display
 	 */
-	 FName GetIconName() const;
+	 UE_API FName GetIconName() const;
 
 	/**
 	 * Get the name of the small icon graphic we should use to display the state in a UI.
 	 * @returns the name of the icon to display
 	 */
-	FName GetSmallIconName() const;
+	UE_API FName GetSmallIconName() const;
 
 	/**
 	 * Get a text representation of the state
 	 * @returns	the text to display for this state
 	 */
-	const FText& GetDisplayText() const;
+	UE_API const FText& GetDisplayText() const;
 
 	/**
 	 * Get a text representation of the state
 	 * @returns	the text to display for this state
 	 */
-	const FText& GetDescriptionText() const;
+	UE_API const FText& GetDescriptionText() const;
 
 	/**
 	 * Get a tooltip to describe this state
 	 * @returns	the text to display for this states tooltip
 	 */
-	FText GetDisplayTooltip() const;
+	UE_API FText GetDisplayTooltip() const;
 
 	/**
 	 * Get the timestamp of the last update that was made to this state.
 	 * @returns	the timestamp of the last update
 	 */
-	const FDateTime& GetTimeStamp() const;
+	UE_API const FDateTime& GetTimeStamp() const;
 
-	const TSet<FSourceControlStateRef>& GetFilesStates() const;
+	UE_API const TSet<FSourceControlStateRef>& GetFilesStates() const;
 	
-	const TSet<FString>& GetOfflineFiles() const;
+	UE_API const TSet<FString>& GetOfflineFiles() const;
 
-	const TSet<FString>& GetDeletedOfflineFiles() const;
+	UE_API const TSet<FString>& GetDeletedOfflineFiles() const;
 
 	/**
 	 * Get the number of files in the CL. (Includes file states and offline files) 
 	 */
-	int32 GetFileCount() const;
+	UE_API int32 GetFileCount() const;
 
 	/**
 	 * Get the filenames in the CL. (Includes file states and offline files)
 	 */
-	TArray<FString> GetFilenames() const;
+	UE_API TArray<FString> GetFilenames() const;
 
 	/**
 	 * Check whether a file exists in the file states or offline files
 	 */
-	bool ContainsFilename(const FString& PackageFilename) const;
+	UE_API bool ContainsFilename(const FString& PackageFilename) const;
 
 	/**
 	 * Serialize the state of the Uncontrolled Changelist to a Json Object, optionally filtering the files.
 	 * @param 	OutJsonObject 	The Json object used to serialize.
 	 * @param	FilenameFilter	An optional filter to control while filenames are included in the serialized JSON.
 	 */
-	void Serialize(TSharedRef<class FJsonObject> OutJsonObject, const TFunction<bool(const FString&)>& FilenameFilter = nullptr) const;
+	UE_API void Serialize(TSharedRef<class FJsonObject> OutJsonObject, const TFunction<bool(const FString&)>& FilenameFilter = nullptr) const;
 
 	/**
 	 * Deserialize the state of the Uncontrolled Changelist from a Json Object.
 	 * @param 	InJsonValue 	The Json Object to read from.
 	 * @return 	True if Deserialization succeeded.
 	 */
-	bool Deserialize(const TSharedRef<FJsonObject> InJsonValue);
+	UE_API bool Deserialize(const TSharedRef<FJsonObject> InJsonValue);
 
 	/**
 	 * Adds files to this Uncontrolled Changelist State.
@@ -132,42 +134,42 @@ public:
 	 * @param 	InCheckFlags 	Tells which checks have to pass to add a file.
 	 * @return 	True if a change has been performed in the Uncontrolled Changelist State.
 	 */
-	bool AddFiles(const TArray<FString>& InFilenames, const ECheckFlags InCheckFlags);
+	UE_API bool AddFiles(const TArray<FString>& InFilenames, const ECheckFlags InCheckFlags);
 
 	/**
 	 * Removes files from this Uncontrolled Changelist State if present.
 	 * @param 	InFileStates 	The files to be removed.
 	 * @return 	True if a change has been performed in the Uncontrolled Changelist State.
 	 */
-	bool RemoveFiles(const TArray<FSourceControlStateRef>& InFileStates);
+	UE_API bool RemoveFiles(const TArray<FSourceControlStateRef>& InFileStates);
 
 	/**
 	 * Removes the given filenames from the file states or offline files.
 	 * @param 	InFilenames 	The filenames to be removed.
 	 * @return 	True if a change has been performed in the Uncontrolled Changelist State.
 	 */
-	bool RemoveFiles(const TArray<FString>& InFilenames);
+	UE_API bool RemoveFiles(const TArray<FString>& InFilenames);
 
 	/**
 	 * Updates the status of all files contained in this changelist.
 	 * @return 	True if the state has been modified.
 	 */
-	bool UpdateStatus();
+	UE_API bool UpdateStatus();
 
 	/**
 	 * Removes files present both in the Uncontrolled Changelist and the provided set.
 	 * @param 	InOutAddedAssets 	A Set representing Added Assets to check against.
 	 */
-	void RemoveDuplicates(TSet<FString>& InOutAddedAssets);
+	UE_API void RemoveDuplicates(TSet<FString>& InOutAddedAssets);
 
 	/**
 	 * Sets a new description for this Uncontrolled Changelist
 	 * @param	InDescription	The new description to set.
 	 */
-	void SetDescription(const FText& InDescription);
+	UE_API void SetDescription(const FText& InDescription);
 
 	/** Returns true if the Uncontrolled Changelists contains either Files or OfflineFiles.	*/
-	bool ContainsFiles() const;
+	UE_API bool ContainsFiles() const;
 
 public:
 	FUncontrolledChangelist Changelist;
@@ -183,3 +185,5 @@ public:
 ENUM_CLASS_FLAGS(FUncontrolledChangelistState::ECheckFlags);
 typedef TSharedPtr<FUncontrolledChangelistState, ESPMode::ThreadSafe> FUncontrolledChangelistStatePtr;
 typedef TSharedRef<FUncontrolledChangelistState, ESPMode::ThreadSafe> FUncontrolledChangelistStateRef;
+
+#undef UE_API
diff --git a/Engine/Source/Developer/UncontrolledChangelists/Public/UncontrolledChangelistsModule.h b/Engine/Source/Developer/UncontrolledChangelists/Public/UncontrolledChangelistsModule.h
--- a/Engine/Source/Developer/UncontrolledChangelists/Public/UncontrolledChangelistsModule.h
+++ b/Engine/Source/Developer/UncontrolledChangelists/Public/UncontrolledChangelistsModule.h
@@ -10,6 +10,8 @@
 #include "UObject/UObjectGlobals.h"
 #include "Async/AsyncWork.h"
 
+#define UE_API UNCONTROLLEDCHANGELISTS_API
+
 struct FAssetData;
 struct FSourceControlProjectInfo;
 
@@ -19,7 +21,7 @@ class FUncontrolledChangelistsDiscoverAssetsTask;
 /**
  * Interface for talking to Uncontrolled Changelists
  */
-class UNCONTROLLEDCHANGELISTS_API FUncontrolledChangelistsModule : public IModuleInterface
+class FUncontrolledChangelistsModule : public IModuleInterface
 {
 	typedef TMap<FUncontrolledChangelist, FUncontrolledChangelistStateRef> FUncontrolledChangelistsStateCache;
 
@@ -33,70 +35,70 @@ public:
 	FOnUncontrolledChangelistModuleChanged OnUncontrolledChangelistModuleChanged;
 
 public:
-	FUncontrolledChangelistsModule();
-	~FUncontrolledChangelistsModule();
+	UE_API FUncontrolledChangelistsModule();
+	UE_API ~FUncontrolledChangelistsModule();
 
 	/** IModuleInterface implementation */
-	virtual void StartupModule() override;
-	virtual void ShutdownModule() override;
+	UE_API virtual void StartupModule() override;
+	UE_API virtual void ShutdownModule() override;
 
 	/**
 	 * Check whether uncontrolled changelist module is enabled.
 	 */
-	bool IsEnabled() const;
+	UE_API bool IsEnabled() const;
 
 	/**
 	 * Get the changelist state of each cached Uncontrolled Changelist.
 	 */
-	TArray<FUncontrolledChangelistStateRef> GetChangelistStates() const;
+	UE_API TArray<FUncontrolledChangelistStateRef> GetChangelistStates() const;
 
 	/**
 	 * Get the changelist state of the given Uncontrolled Changelist.
 	 */
-	FUncontrolledChangelistStatePtr GetChangelistState(const FUncontrolledChangelist& InUncontrolledChangelist) const;
+	UE_API FUncontrolledChangelistStatePtr GetChangelistState(const FUncontrolledChangelist& InUncontrolledChangelist) const;
 
 	/**
 	 * Get the changelist state of the default Uncontrolled Changelist.
 	 */
-	FUncontrolledChangelistStatePtr GetDefaultChangelistState() const;
+	UE_API FUncontrolledChangelistStatePtr GetDefaultChangelistState() const;
 
 	/**
 	 * Called if the state of any Uncontrolled Changelist is modified externally, eg, via the mutable accessors above.
 	 */
-	void HandleChangelistStateModified();
+	UE_API void HandleChangelistStateModified();
 
 	/**
 	 * Called when file has been made writable. Adds the file to the reconcile list because we don't know yet if it will change.
 	 * @param	InFilename		The file to be reconciled.
 	 * @return True if the file have been handled by the Uncontrolled module.
 	 */
-	bool OnMakeWritable(const FString& InFilename);
+	UE_API bool OnMakeWritable(const FString& InFilename);
 	 	 
 	/**
 	 * Called when file has been saved without an available Provider. Adds the file to the Default Uncontrolled Changelist
 	 * @param	InFilename		The file to be added.
 	 * @return	True if the file have been handled by the Uncontrolled module.
 	 */
-	bool OnSaveWritable(const FString& InFilename);
+	UE_API bool OnSaveWritable(const FString& InFilename);
 
 	/**
 	 * Called when file has been deleted without an available Provider. Adds the file to the Default Uncontrolled Changelist
 	 * @param	InFilename		The file to be added.
 	 * @return	True if the file have been handled by the Uncontrolled module.
 	 */
-	bool OnDeleteWritable(const FString& InFilename);
+	UE_API bool OnDeleteWritable(const FString& InFilename);
 
 	/**
 	 * Called when files should have been marked for add without an available Provider. Adds the files to the Default Uncontrolled Changelist
 	 * @param	InFilenames		The files to be added.
 	 * @return	True if the files have been handled by the Uncontrolled module.
 	 */
-	bool OnNewFilesAdded(const TArray<FString>& InFilenames);
+	UE_API bool OnNewFilesAdded(const TArray<FString>& InFilenames);
 
 	/**
 	 * Updates the status of Uncontrolled Changelists and files.
 	 */
-	void UpdateStatus();
+	UE_API void UpdateStatus();
 
 	/**
 	 * Gets a reference to the UncontrolledChangelists module
@@ -136,32 +138,32 @@ public:
 	 * Gets a message indicating the status of SCC coherence.
 	 * @return 	A text representing the status of SCC.
 	 */
-	FText GetReconcileStatus() const;
+	UE_API FText GetReconcileStatus() const;
 
 	/** Called when "Reconcile assets" button is clicked. Checks for uncontrolled modifications in previously added assets.
 	 *	Adds modified files to Uncontrolled Changelists
 	 *  @return True if new modifications found
 	 */
-	bool OnReconcileAssets();
+	UE_API bool OnReconcileAssets();
 
 	/**
 	 * Delegate callback called when assets are added to AssetRegistry.
 	 * @param 	AssetData 	The asset just added.
 	 */
-	void OnAssetAdded(const FAssetData& AssetData);
+	UE_API void OnAssetAdded(const FAssetData& AssetData);
 
 	/** Called when "Revert files" button is clicked. Reverts modified files and deletes new ones.
 	 *  @param	InFilenames		The files to be reverted
 	 *	@return true if the provided files were reverted.
 	 */
-	bool OnRevert(const TArray<FString>& InFilenames);
+	UE_API bool OnRevert(const TArray<FString>& InFilenames);
 	
 	/**
 	 * Delegate callback called before an asset has been written to disk.
 	 * @param 	InObject 			The saved object.
 	 * @param 	InPreSaveContext 	Interface used to access saved parameters.
 	 */
-	void OnObjectPreSaved(UObject* InObject, FObjectPreSaveContext InPreSaveContext);
+	UE_API void OnObjectPreSaved(UObject* InObject, FObjectPreSaveContext InPreSaveContext);
 
 	/**
 	 * Moves files to an Uncontrolled Changelist.
@@ -169,14 +171,14 @@ public:
 	 * @param 	InUncontrolledFileStates 	The Uncontrolled files to move.
 	 * @param 	InUncontrolledChangelist 	The Uncontrolled Changelist where to move the files.
 	 */
-	void MoveFilesToUncontrolledChangelist(const TArray<FSourceControlStateRef>& InControlledFileStates, const TArray<FSourceControlStateRef>& InUncontrolledFileStates, const FUncontrolledChangelist& InUncontrolledChangelist);
+	UE_API void MoveFilesToUncontrolledChangelist(const TArray<FSourceControlStateRef>& InControlledFileStates, const TArray<FSourceControlStateRef>& InUncontrolledFileStates, const FUncontrolledChangelist& InUncontrolledChangelist);
 
 	/**
 	* Moves files to an Uncontrolled Changelist.
 	* @param 	InControlledFileStates 		The Controlled files to move.
 	* @param 	InUncontrolledChangelist 	The Uncontrolled Changelist where to move the files.
 	*/
-	void MoveFilesToUncontrolledChangelist(const TArray<FString>& InControlledFiles, const FUncontrolledChangelist& InUncontrolledChangelist);
+	UE_API void MoveFilesToUncontrolledChangelist(const TArray<FString>& InControlledFiles, const FUncontrolledChangelist& InUncontrolledChangelist);
 
 	/**
 	 * Moves files to a Controlled Changelist.
@@ -184,7 +186,7 @@ public:
 	 * @param 	InChangelist 				The Controlled Changelist where to move the files.
 	 * @param 	InOpenConflictDialog 		A callback to be used by the method when file conflicts are detected. The callback should display the files and ask the user if they should proceed.
 	 */
-	void MoveFilesToControlledChangelist(const TArray<FSourceControlStateRef>& InUncontrolledFileStates, const FSourceControlChangelistPtr& InChangelist, TFunctionRef<bool(const TArray<FSourceControlStateRef>&)> InOpenConflictDialog);
+	UE_API void MoveFilesToControlledChangelist(const TArray<FSourceControlStateRef>& InUncontrolledFileStates, const FSourceControlChangelistPtr& InChangelist, TFunctionRef<bool(const TArray<FSourceControlStateRef>&)> InOpenConflictDialog);
 	
 	/**
 	 * Moves files to a Controlled Changelist.
@@ -192,7 +194,7 @@ public:
 	 * @param 	InChangelist 			The Controlled Changelist where to move the files.
 	 * @param 	InOpenConflictDialog 	A callback to be used by the method when file conflicts are detected. The callback should display the files and ask the user if they should proceed.
 	 */
-	void MoveFilesToControlledChangelist(const TArray<FString>& InUncontrolledFiles, const FSourceControlChangelistPtr& InChangelist, TFunctionRef<bool(const TArray<FSourceControlStateRef>&)> InOpenConflictDialog);
+	UE_API void MoveFilesToControlledChangelist(const TArray<FString>& InUncontrolledFiles, const FSourceControlChangelistPtr& InChangelist, TFunctionRef<bool(const TArray<FSourceControlStateRef>&)> InOpenConflictDialog);
 
 	/**
 	 * Creates a new Uncontrolled Changelist.
@@ -200,20 +202,20 @@ public:
 	 * @param	InUncontrolledChangelist	An optional Uncontrolled Changelist to create (or find) via its ID. Should not be the default Uncontrolled Changelist.
 	 * return	TOptional<FUncontrolledChangelist> set with the new Uncontrolled Changelist key if succeeded.
 	 */
-	TOptional<FUncontrolledChangelist> CreateUncontrolledChangelist(const FText& InDescription, const TOptional<FUncontrolledChangelist>& InUncontrolledChangelist = TOptional<FUncontrolledChangelist>());
+	UE_API TOptional<FUncontrolledChangelist> CreateUncontrolledChangelist(const FText& InDescription, const TOptional<FUncontrolledChangelist>& InUncontrolledChangelist = TOptional<FUncontrolledChangelist>());
 
 	/**
      * Edits an Uncontrolled Changelist's description
 	 * @param	InUncontrolledChangelist	The Uncontrolled Changelist to modify. Should not be the default Uncontrolled Changelist.
 	 * @param	InNewDescription			The description to set.
 	 */
-	void EditUncontrolledChangelist(const FUncontrolledChangelist& InUncontrolledChangelist, const FText& InNewDescription);
+	UE_API void EditUncontrolledChangelist(const FUncontrolledChangelist& InUncontrolledChangelist, const FText& InNewDescription);
 	
 	/**
 	 * Deletes an Uncontrolled Changelist.
 	 * @param	InUncontrolledChangelist	The Uncontrolled Changelist to delete. Should not be the default Uncontrolled Changelist and should not contain files.
 	 */
-	void DeleteUncontrolledChangelist(const FUncontrolledChangelist& InUncontrolledChangelist);
+	UE_API void DeleteUncontrolledChangelist(const FUncontrolledChangelist& InUncontrolledChangelist);
 
 private:
 	/**
@@ -223,107 +225,107 @@ private:
 	 * @param	bInDiscoveryTask	If true, this asset was added from the asset discovery task.
 	 * 
 	 */
-	void OnAssetAddedInternal(const FAssetData& InAssetData, TSet<FString>& InAddedAssetsCache, bool bInDiscoveryTask);
+	UE_API void OnAssetAddedInternal(const FAssetData& InAssetData, TSet<FString>& InAddedAssetsCache, bool bInDiscoveryTask);
 
 	/**
 	 * Add files to Uncontrolled Changelist. Also adds them to files to reconcile.
 	 */
-	bool AddToUncontrolledChangelist(const TArray<FString>& InFilenames);
+	UE_API bool AddToUncontrolledChangelist(const TArray<FString>& InFilenames);
 
 	/**
 	 * Removes the given files from their associated Uncontrolled Changelist, if any.  Also removes them from files to reconcile.
 	 * @return true if any of the given files are removed.
 	 */
-	bool RemoveFromUncontrolledChangelist(const TArray<FString>& InFilenames);
+	UE_API bool RemoveFromUncontrolledChangelist(const TArray<FString>& InFilenames);
 
 	/**
 	 * Groups the given files by their associated Uncontrolled Changelist and returns the groupings in the given map.
 	 * Any of the files that are not already associated with an Uncontrolled Changelist will be added to the default Uncontrolled Changelist's group.
 	 */
-	void GroupFilesByUncontrolledChangelist(TArray<FString>&& InFilenames, TMap<FUncontrolledChangelist, TArray<FString>>& OutUncontrolledChangelistToFilenames) const;
+	UE_API void GroupFilesByUncontrolledChangelist(TArray<FString>&& InFilenames, TMap<FUncontrolledChangelist, TArray<FString>>& OutUncontrolledChangelistToFilenames) const;
 
 	/**
 	 * True if we have custom projects and calling DoesFilePassCustomProjectFilter could ever return something aside from true.
 	 */
-	bool HasCustomProjectFilter() const;
+	UE_API bool HasCustomProjectFilter() const;
 
 	/**
 	 * Run the given file through the filter for all known custom projects.
 	 */
-	bool DoesFilePassCustomProjectFilter(const FString& InFilename) const;
+	UE_API bool DoesFilePassCustomProjectFilter(const FString& InFilename) const;
 
 	/**
 	 * Run the given file through the filter for the given custom project.
 	 */
-	static bool DoesFilePassCustomProjectFilter(const FString& InFilename, const FSourceControlProjectInfo& Project);
+	static UE_API bool DoesFilePassCustomProjectFilter(const FString& InFilename, const FSourceControlProjectInfo& Project);
 
 	/**
 	 * Saves the state of UncontrolledChangelists to Json for persistency.
 	 */
-	void SaveState();
+	UE_API void SaveState();
 	
 	/**
 	 * Restores the previously saved state from Json.
 	 */
-	void LoadState();
+	UE_API void LoadState();
 	
 	/**
 	 * Request that ReloadState be called at the end of the current frame.
 	 */
-	void RequestReloadState();
+	UE_API void RequestReloadState();
 
 	/**
 	 * Reload the state of the UncontrolledChangelists when ISourceControlModule::GetCustomProjects changes.
 	 */
-	void ReloadState();
+	UE_API void ReloadState();
 
 	/**
 	 * Removes any duplicated files across changelists
 	 */
-	void SanitizeState();
+	UE_API void SanitizeState();
 
 	/**
 	 * Called on End of frame. Calls SaveState if needed.
 	 */
-	void OnEndFrame();
+	UE_API void OnEndFrame();
 
 	/**
 	 * Called when the uncontrolled changelist feature switches to an enabled state.
 	 */
-	void OnEnabled();
+	UE_API void OnEnabled();
 
 	/**
 	 * Called when the uncontrolled changelist feature switches to a disabled state.
 	 */
-	void OnDisabled();
+	UE_API void OnDisabled();
 
 	/**
 	 * Start a new DiscoverAssetsTask for the current state.
 	 * @note InitialScanEvent must be null before calling this, and there shouldn't be any existing DiscoverAssetsTask running.
 	 */
-	void StartAssetDiscovery();
+	UE_API void StartAssetDiscovery();
 
 	/**
 	 * Stop the current DiscoverAssetsTask, if any.
 	 */
-	void StopAssetDiscovery();
+	UE_API void StopAssetDiscovery();
 
 	/**
 	 * Used by the task to query whether we're waiting for it to finish in StopAssetDiscovery. 
 	 */
-	bool IsStopAssetDiscoveryRequested() const;
+	UE_API bool IsStopAssetDiscoveryRequested() const;
 
 	/**
 	 * Helper returning the location of the file used for persistency.
 	 * @return 	A string containing the filepath.
 	 */
-	FString GetPersistentFilePath(const FString& SubProjectName) const;
+	UE_API FString GetPersistentFilePath(const FString& SubProjectName) const;
 
 	/** Called when a state changed either in the module or an Uncontrolled Changelist. */
-	void OnStateChanged();
+	UE_API void OnStateChanged();
 
 	/** Removes from asset caches files already present in Uncontrolled Changelists */
-	void CleanAssetsCaches();
+	UE_API void CleanAssetsCaches();
 
 	/**
 	 * Try to add the provided filenames to the given Uncontrolled Changelist.
@@ -332,7 +334,7 @@ private:
 	 * @param	InCheckFlags 				The required checks to check the file against before adding.
 	 * @return	True files have been added.
 	 */
-	bool AddFilesToUncontrolledChangelist(const FUncontrolledChangelist& InUncontrolledChangelist, const TArray<FString>& InFilenames, const FUncontrolledChangelistState::ECheckFlags InCheckFlags);
+	UE_API bool AddFilesToUncontrolledChangelist(const FUncontrolledChangelist& InUncontrolledChangelist, const TArray<FString>& InFilenames, const FUncontrolledChangelistState::ECheckFlags InCheckFlags);
 
 	/**
 	 * Try to remove the provided filenames from the given Uncontrolled Changelist.
@@ -340,17 +342,17 @@ private:
 	 * @param	InFilenames 				The files to remove.
 	 * @return	True files have been removed.
 	 */
-	bool RemoveFilesFromUncontrolledChangelist(const FUncontrolledChangelist& InUncontrolledChangelist, const TArray<FString>& InFilenames);
+	UE_API bool RemoveFilesFromUncontrolledChangelist(const FUncontrolledChangelist& InUncontrolledChangelist, const TArray<FString>& InFilenames);
 
 	/** Returns the default Uncontrolled Changelist state, creates it if it does not exist. */
-	FUncontrolledChangelistStateRef GetDefaultUncontrolledChangelistState();
+	UE_API FUncontrolledChangelistStateRef GetDefaultUncontrolledChangelistState();
 
 	/**
 	 * Returns the given Uncontrolled Changelist state if it exists.
 	 * If the given Uncontrolled Changelist is the default one and it does not exist, then it will be created and returned.
 	 * Otherwise, if not the default Uncontrolled Changelist, then nullptr will be returned if it does not exist.
 	 */
-	FUncontrolledChangelistStatePtr GetUncontrolledChangelistState(const FUncontrolledChangelist& InUncontrolledChangelist);
+	UE_API FUncontrolledChangelistStatePtr GetUncontrolledChangelistState(const FUncontrolledChangelist& InUncontrolledChangelist);
 
 private:
 	friend FUncontrolledChangelistsDiscoverAssetsTask;
@@ -374,3 +376,5 @@ private:
 	bool								bIsStateDirty = false;
 	bool								bPendingReloadState = false;
 };
+
+#undef UE_API
diff --git a/Engine/Source/Developer/UndoHistory/Public/Widgets/IReflectionDataProvider.h b/Engine/Source/Developer/UndoHistory/Public/Widgets/IReflectionDataProvider.h
--- a/Engine/Source/Developer/UndoHistory/Public/Widgets/IReflectionDataProvider.h
+++ b/Engine/Source/Developer/UndoHistory/Public/Widgets/IReflectionDataProvider.h
@@ -33,7 +33,7 @@ namespace UE::UndoHistory
 	 * On editor builds reflection data is fully available however on non-editor builds there is no straightforward way of obtaining it.
 	 * For example, Concert implements this interface to return data where available.
 	 */
-	class UNDOHISTORY_API IReflectionDataProvider
+	class IReflectionDataProvider
 	{
 	public:
 
diff --git a/Engine/Source/Developer/UndoHistory/Public/Widgets/SUndoHistoryDetails.h b/Engine/Source/Developer/UndoHistory/Public/Widgets/SUndoHistoryDetails.h
--- a/Engine/Source/Developer/UndoHistory/Public/Widgets/SUndoHistoryDetails.h
+++ b/Engine/Source/Developer/UndoHistory/Public/Widgets/SUndoHistoryDetails.h
@@ -9,7 +9,9 @@
 #include "Widgets/Views/STreeView.h"
 #include "Misc/TextFilter.h"
 
-class UNDOHISTORY_API SUndoHistoryDetails : public SCompoundWidget
+#define UE_API UNDOHISTORY_API
+
+class SUndoHistoryDetails : public SCompoundWidget
 {
 
 public:
@@ -23,21 +25,21 @@ public:
 	 *
 	 * @param InArgs The declaration data for this widget.
 	 */
-	void Construct(const FArguments& InArgs, TSharedRef<UE::UndoHistory::IReflectionDataProvider> InReflectionData = UE::UndoHistory::CreateDefaultReflectionProvider());
+	UE_API void Construct(const FArguments& InArgs, TSharedRef<UE::UndoHistory::IReflectionDataProvider> InReflectionData = UE::UndoHistory::CreateDefaultReflectionProvider());
 
 	/**
 	 * Set the transaction to be displayed in the details panel.
 	 */
-	void SetSelectedTransaction(const struct FTransactionDiff& InTransactionDiff);
+	UE_API void SetSelectedTransaction(const struct FTransactionDiff& InTransactionDiff);
 
 	/**
 	 * Clear the details panel.
 	 */
-	void Reset();
+	UE_API void Reset();
 
 
 	/** SWidget interface */
-	virtual void Tick(const FGeometry& AllottedGeometry, const double InCurrentTime, const float InDeltaTime) override;
+	UE_API virtual void Tick(const FGeometry& AllottedGeometry, const double InCurrentTime, const float InDeltaTime) override;
 
 private:
 
@@ -67,43 +69,43 @@ private:
 private:
 
 	/** Create a changed object node. */
-	FUndoDetailsTreeNodePtr CreateTreeNode(const FString& InObjectPathName, const FSoftClassPath& InObjectClass, const TSharedPtr<FTransactionObjectEvent>& InEvent) const;
+	UE_API FUndoDetailsTreeNodePtr CreateTreeNode(const FString& InObjectPathName, const FSoftClassPath& InObjectClass, const TSharedPtr<FTransactionObjectEvent>& InEvent) const;
 
 	/** Create a tooltip text for a property. */
-	FText CreateToolTipText(EPropertyFlags InFlags) const;
+	UE_API FText CreateToolTipText(EPropertyFlags InFlags) const;
 
 	/** Whether the type row will be generated */
-	bool SupportsTypeRow() const;
+	UE_API bool SupportsTypeRow() const;
 
 	/** Callback to handle a change in the filter box. */
-	void OnFilterTextChanged(const FText& InFilterText);
+	UE_API void OnFilterTextChanged(const FText& InFilterText);
 
 	/** Refresh the details tree view. */
-	void FullRefresh();
+	UE_API void FullRefresh();
 
 	/** Populate the search strings for the filter. */
-	void PopulateSearchStrings(const FString&, TArray<FString>& OutSearchStrings) const;
+	UE_API void PopulateSearchStrings(const FString&, TArray<FString>& OutSearchStrings) const;
 
 	/** Populate the details tree. */
-	void Populate();
+	UE_API void Populate();
 
 	/** Callback for generating a UndoHistoryDetailsRow */
-	TSharedRef<ITableRow> HandleGenerateRow(FUndoDetailsTreeNodePtr InNode, const TSharedRef<STableViewBase>& OwnerTable) const;
+	UE_API TSharedRef<ITableRow> HandleGenerateRow(FUndoDetailsTreeNodePtr InNode, const TSharedRef<STableViewBase>& OwnerTable) const;
 
 	/** Callback for getting the filter highlight text. */
-	FText HandleGetFilterHighlightText() const;
+	UE_API FText HandleGetFilterHighlightText() const;
 
 	/** Callback for getting the details visibility. */
-	EVisibility HandleDetailsVisibility() const;
+	UE_API EVisibility HandleDetailsVisibility() const;
 
 	/** Callback for getting the transaction name. */
-	FText HandleTransactionName() const;
+	UE_API FText HandleTransactionName() const;
 
 	/* Callback for getting the transaction id. */
-	FText HandleTransactionId() const;
+	UE_API FText HandleTransactionId() const;
 
 	/** Callback for handling a click on the transaction id. */
-	void HandleTransactionIdNavigate();
+	UE_API void HandleTransactionIdNavigate();
 
 private:
 
@@ -137,3 +139,5 @@ private:
 	/** If the tree items need to be expanded (ie. When the filter text changes). */
 	bool bNeedsExpansion;
 };
+
+#undef UE_API
diff --git a/Engine/Source/Developer/UnsavedAssetsTracker/Source/Public/UnsavedAssetsTrackerModule.h b/Engine/Source/Developer/UnsavedAssetsTracker/Source/Public/UnsavedAssetsTrackerModule.h
--- a/Engine/Source/Developer/UnsavedAssetsTracker/Source/Public/UnsavedAssetsTrackerModule.h
+++ b/Engine/Source/Developer/UnsavedAssetsTracker/Source/Public/UnsavedAssetsTrackerModule.h
@@ -11,6 +11,8 @@
 #include "Templates/SharedPointer.h"
 #include "UObject/ObjectSaveContext.h"
 
+#define UE_API UNSAVEDASSETSTRACKER_API
+
 class FUnsavedAssetsTracker;
 class FUnsavedAssetsAutoCheckout;
 class SWidget;
@@ -31,12 +33,12 @@ DECLARE_MULTICAST_DELEGATE_TwoParams(FPostUnsavedAssetAutoCheckout, const FStrin
  * Tracks assets that has in-memory modification not saved to disk yet and checks
  * the source control states of those assets when a source control provider is available.
  */
-class UNSAVEDASSETSTRACKER_API FUnsavedAssetsTrackerModule : public IModuleInterface
+class FUnsavedAssetsTrackerModule : public IModuleInterface
 {
 public:
 	/** IModuleInterface implementation */
-	virtual void StartupModule() override;
-	virtual void ShutdownModule() override;
+	UE_API virtual void StartupModule() override;
+	UE_API virtual void ShutdownModule() override;
 
 	/**
 	 * Retrieve the module instance.
@@ -47,19 +49,19 @@ public:
 	}
 
 	/** Returns the number of unsaved assets currently tracked. */
-	int32 GetUnsavedAssetNum() const;
+	UE_API int32 GetUnsavedAssetNum() const;
 
 	/** Returns the list of unsaved assets. */
-	TArray<FString> GetUnsavedAssets() const;
+	UE_API TArray<FString> GetUnsavedAssets() const;
 
 	/** Check if the input asset is unsaved. */
-	bool IsAssetUnsaved(const FString& FileAbsPathname) const;
+	UE_API bool IsAssetUnsaved(const FString& FileAbsPathname) const;
 
 	/** Construct the widget for the Editor status bar.*/
-	TSharedRef<SWidget> MakeUnsavedAssetsStatusBarWidget();
+	UE_API TSharedRef<SWidget> MakeUnsavedAssetsStatusBarWidget();
 
 	/** Displays a dialog prompting the user to save unsaved packages. */
-	bool PromptToSavePackages();
+	UE_API bool PromptToSavePackages();
 
 	/** Invoked when a file is added to the unsaved list. */
 	FOnUnsavedAssetAdded OnUnsavedAssetAdded;
@@ -84,3 +86,5 @@ private:
 	TSharedPtr<FUnsavedAssetsTracker> UnsavedAssetTracker;
 	TSharedPtr<FUnsavedAssetsAutoCheckout> UnsavedAssetAutoCheckout;
 };
+
+#undef UE_API
diff --git a/Engine/Source/Developer/VisualGraphUtils/Public/VisualGraph.h b/Engine/Source/Developer/VisualGraphUtils/Public/VisualGraph.h
--- a/Engine/Source/Developer/VisualGraphUtils/Public/VisualGraph.h
+++ b/Engine/Source/Developer/VisualGraphUtils/Public/VisualGraph.h
@@ -15,9 +15,11 @@
 #include "VisualGraphElement.h"
 #include "VisualGraphNode.h"
 
+#define UE_API VISUALGRAPHUTILS_API
+
 class FVisualGraph;
 
-class VISUALGRAPHUTILS_API FVisualGraphSubGraph : public FVisualGraphElement
+class FVisualGraphSubGraph : public FVisualGraphElement
 {
 public:
 
@@ -33,7 +35,7 @@ public:
 
 protected:
 
-	virtual FString DumpDot(const FVisualGraph* InGraph, int32 InIndendation) const override;
+	UE_API virtual FString DumpDot(const FVisualGraph* InGraph, int32 InIndendation) const override;
 
 	int32 ParentGraphIndex;
 	TArray<int32> Nodes;
@@ -41,7 +43,7 @@ protected:
 	friend class FVisualGraph;
 };
 
-class VISUALGRAPHUTILS_API FVisualGraph : public FVisualGraphElement
+class FVisualGraph : public FVisualGraphElement
 {
 public:
 
@@ -51,20 +53,20 @@ public:
 
 	virtual ~FVisualGraph() override {}
 
-	FVisualGraph(const FName& InName, const FName& InDisplayName = NAME_None);
+	UE_API FVisualGraph(const FName& InName, const FName& InDisplayName = NAME_None);
 
 	const TArray<FVisualGraphNode>& GetNodes() const { return Nodes; }
 	const TArray<FVisualGraphEdge>& GetEdges() const { return Edges; }
 	const TArray<FVisualGraphSubGraph>& GetSubGraphs() const { return SubGraphs; }
 
-	int32 AddNode(
+	UE_API int32 AddNode(
 		const FName& InName,
 		TOptional<FName> InDisplayName = TOptional<FName>(),
 		TOptional<FLinearColor> InColor = TOptional<FLinearColor>(),
 		TOptional<EVisualGraphShape> InShape = TOptional<EVisualGraphShape>(), 
 		TOptional<EVisualGraphStyle> InStyle = TOptional<EVisualGraphStyle>());
 
-	int32 AddEdge(
+	UE_API int32 AddEdge(
 		int32 InSourceNode,
 		int32 InTargetNode,
 		EVisualGraphEdgeDirection InDirection,
@@ -73,7 +75,7 @@ public:
 		TOptional<FLinearColor> InColor = TOptional<FLinearColor>(),
 		TOptional<EVisualGraphStyle> InStyle = TOptional<EVisualGraphStyle>());
 
-	int32 AddSubGraph(
+	UE_API int32 AddSubGraph(
 		const FName& InName,
 		TOptional<FName> InDisplayName = TOptional<FName>(),
 		int32 InParentGraphIndex = INDEX_NONE,
@@ -81,14 +83,14 @@ public:
 		TOptional<EVisualGraphStyle> InStyle = TOptional<EVisualGraphStyle>(),
 		const TArray<int32> InNodes = TArray<int32>());
 
-	int32 FindNode(const FName& InName) const;
-	int32 FindEdge(const FName& InName) const;
-	int32 FindSubGraph(const FName& InName) const;
+	UE_API int32 FindNode(const FName& InName) const;
+	UE_API int32 FindEdge(const FName& InName) const;
+	UE_API int32 FindSubGraph(const FName& InName) const;
 
-	bool AddNodeToSubGraph(int32 InNodeIndex, int32 InSubGraphIndex);
-	bool RemoveNodeFromSubGraph(int32 InNodeIndex);
+	UE_API bool AddNodeToSubGraph(int32 InNodeIndex, int32 InSubGraphIndex);
+	UE_API bool RemoveNodeFromSubGraph(int32 InNodeIndex);
 
-	void TransitiveReduction(TFunction<bool(FVisualGraphEdge&)> KeepEdgeFunction);
+	UE_API void TransitiveReduction(TFunction<bool(FVisualGraphEdge&)> KeepEdgeFunction);
 
 	FORCEINLINE FString DumpDot() const
 	{
@@ -97,7 +99,7 @@ public:
 
 protected:
 
-	virtual FString DumpDot(const FVisualGraph* InGraph, int32 InIndendation) const override;
+	UE_API virtual FString DumpDot(const FVisualGraph* InGraph, int32 InIndendation) const override;
 
 	template<typename T>
 	FORCEINLINE static void RefreshNameMap(const TArray<T>& InElements, TMap<FName, int32>& OutMap)
@@ -131,8 +133,8 @@ protected:
 		return AddedIndex;
 	}
 
-	static bool IsNameAvailable(const FName& InName, const TMap<FName, int32>& InMap);
-	static FName GetUniqueName(const FName& InName, const TMap<FName, int32>& InMap);
+	static UE_API bool IsNameAvailable(const FName& InName, const TMap<FName, int32>& InMap);
+	static UE_API FName GetUniqueName(const FName& InName, const TMap<FName, int32>& InMap);
 
 	TArray<FVisualGraphNode> Nodes;
 	TArray<FVisualGraphEdge> Edges;
@@ -142,3 +144,5 @@ protected:
 	TMap<FName,int32> EdgeNameMap;
 	TMap<FName,int32> SubGraphNameMap;
 };
+
+#undef UE_API
diff --git a/Engine/Source/Developer/VisualGraphUtils/Public/VisualGraphEdge.h b/Engine/Source/Developer/VisualGraphUtils/Public/VisualGraphEdge.h
--- a/Engine/Source/Developer/VisualGraphUtils/Public/VisualGraphEdge.h
+++ b/Engine/Source/Developer/VisualGraphUtils/Public/VisualGraphEdge.h
@@ -9,6 +9,8 @@
 
 #include "VisualGraphEdge.generated.h"
 
+#define UE_API VISUALGRAPHUTILS_API
+
 class FVisualGraph;
 
 UENUM()
@@ -19,7 +21,7 @@ enum class EVisualGraphEdgeDirection : uint8
 	BothWays
 };
 
-class VISUALGRAPHUTILS_API FVisualGraphEdge : public FVisualGraphElement
+class FVisualGraphEdge : public FVisualGraphElement
 {
 public:
 
@@ -36,7 +38,7 @@ public:
 
 protected:
 
-	virtual FString DumpDot(const FVisualGraph* InGraph, int32 InIndendation) const override;
+	UE_API virtual FString DumpDot(const FVisualGraph* InGraph, int32 InIndendation) const override;
 
 	int32 SourceNode;
 	int32 TargetNode;
@@ -46,3 +48,4 @@ protected:
 	friend class FVisualGraphSubGraph;
 };
 
+#undef UE_API
diff --git a/Engine/Source/Developer/VisualGraphUtils/Public/VisualGraphElement.h b/Engine/Source/Developer/VisualGraphUtils/Public/VisualGraphElement.h
--- a/Engine/Source/Developer/VisualGraphUtils/Public/VisualGraphElement.h
+++ b/Engine/Source/Developer/VisualGraphUtils/Public/VisualGraphElement.h
@@ -14,6 +14,8 @@
 
 #include "VisualGraphElement.generated.h"
 
+#define UE_API VISUALGRAPHUTILS_API
+
 class FVisualGraph;
 
 UENUM()
@@ -42,7 +44,7 @@ enum class EVisualGraphStyle : uint8
 	Bold
 };
 
-class VISUALGRAPHUTILS_API FVisualGraphElement
+class FVisualGraphElement
 {
 public:
 
@@ -69,10 +71,10 @@ public:
 protected:
 
 	virtual FString DumpDot(const FVisualGraph* InGraph, int32 InIndendation) const = 0;
-	static FString DumpDotIndentation(int32 InIndentation);
-	static FString DumpDotColor(const TOptional<FLinearColor>& InColor);
-	static FString DumpDotShape(const TOptional<EVisualGraphShape>& InShape);
-	static FString DumpDotStyle(const TOptional<EVisualGraphStyle>& InStyle);
+	static UE_API FString DumpDotIndentation(int32 InIndentation);
+	static UE_API FString DumpDotColor(const TOptional<FLinearColor>& InColor);
+	static UE_API FString DumpDotShape(const TOptional<EVisualGraphShape>& InShape);
+	static UE_API FString DumpDotStyle(const TOptional<EVisualGraphStyle>& InStyle);
 
 	FName Name;
 	TOptional<FName> DisplayName;
@@ -84,3 +86,4 @@ protected:
 	friend class FVisualGraph;
 };
 
+#undef UE_API
diff --git a/Engine/Source/Developer/VisualGraphUtils/Public/VisualGraphNode.h b/Engine/Source/Developer/VisualGraphUtils/Public/VisualGraphNode.h
--- a/Engine/Source/Developer/VisualGraphUtils/Public/VisualGraphNode.h
+++ b/Engine/Source/Developer/VisualGraphUtils/Public/VisualGraphNode.h
@@ -8,9 +8,11 @@
 #include "Misc/Optional.h"
 #include "VisualGraphElement.h"
 
+#define UE_API VISUALGRAPHUTILS_API
+
 class FVisualGraph;
 
-class VISUALGRAPHUTILS_API FVisualGraphNode : public FVisualGraphElement
+class FVisualGraphNode : public FVisualGraphElement
 {
 public:
 
@@ -27,7 +29,7 @@ public:
 
 protected:
 
-	virtual FString DumpDot(const FVisualGraph* InGraph, int32 InIndendation) const override;
+	UE_API virtual FString DumpDot(const FVisualGraph* InGraph, int32 InIndendation) const override;
 
 	TOptional<EVisualGraphShape> Shape;
 	int32 SubGraphIndex;
@@ -35,3 +37,5 @@ protected:
 	friend class FVisualGraph;
 	friend class FVisualGraphSubGraph;
 };
+
+#undef UE_API
diff --git a/Engine/Source/Developer/VisualGraphUtils/Public/VisualGraphObjectUtils.h b/Engine/Source/Developer/VisualGraphUtils/Public/VisualGraphObjectUtils.h
--- a/Engine/Source/Developer/VisualGraphUtils/Public/VisualGraphObjectUtils.h
+++ b/Engine/Source/Developer/VisualGraphUtils/Public/VisualGraphObjectUtils.h
@@ -5,14 +5,16 @@
 #include "Containers/Array.h"
 #include "VisualGraph.h"
 
+#define UE_API VISUALGRAPHUTILS_API
+
 class UClass;
 class UObject;
 
-class VISUALGRAPHUTILS_API FVisualGraphObjectUtils
+class FVisualGraphObjectUtils
 {
 public:
 
-	static FVisualGraph TraverseUObjectReferences(
+	static UE_API FVisualGraph TraverseUObjectReferences(
 		const TArray<UObject*>& InObjects,
 		const TArray<UClass*>& InClassesToSkip = TArray<UClass*>(),
 		const TArray<UObject*>& InOutersToSkip = TArray<UObject*>(),
@@ -22,8 +24,9 @@ public:
 		bool bRecursive = true
 		);
 
-	static FVisualGraph TraverseTickOrder(
+	static UE_API FVisualGraph TraverseTickOrder(
 		const TArray<UObject*>& InObjects
 		);
 };
 
+#undef UE_API
diff --git a/Engine/Source/Developer/WidgetRegistration/Public/FToolkitWidgetStyle.h b/Engine/Source/Developer/WidgetRegistration/Public/FToolkitWidgetStyle.h
--- a/Engine/Source/Developer/WidgetRegistration/Public/FToolkitWidgetStyle.h
+++ b/Engine/Source/Developer/WidgetRegistration/Public/FToolkitWidgetStyle.h
@@ -10,11 +10,13 @@
 #include "Fonts/SlateFontInfo.h"
 #include "FToolkitWidgetStyle.generated.h"
 
+#define UE_API WIDGETREGISTRATION_API
+
 /**
  * FToolkitWidgetStyle is the FSlateWidgetStyle that defines the styling of a ToolkitWidget 
  */
 USTRUCT(BlueprintType)
-struct WIDGETREGISTRATION_API FToolkitWidgetStyle : public FSlateWidgetStyle
+struct FToolkitWidgetStyle : public FSlateWidgetStyle
 {
 	GENERATED_BODY()
 
@@ -26,37 +28,39 @@ struct WIDGETREGISTRATION_API FToolkitWidgetStyle : public FSlateWidgetStyle
 		
 	}
 
-	static const FName TypeName;
+	static UE_API const FName TypeName;
 	virtual const FName GetTypeName() const override { return TypeName; }
 
-	static const FToolkitWidgetStyle& GetDefault();
-		virtual void GetResources(TArray<const FSlateBrush*>& OutBrushes) const override;
+	static UE_API const FToolkitWidgetStyle& GetDefault();
+		UE_API virtual void GetResources(TArray<const FSlateBrush*>& OutBrushes) const override;
 
 	UPROPERTY(EditAnywhere, BlueprintReadWrite, Category = Appearance)
 	FSlateBrush TitleBackgroundBrush;
-	FToolkitWidgetStyle& SetTitleBackgroundBrush(const FSlateBrush& InPaletteTitleBackgroundBrush);
+	UE_API FToolkitWidgetStyle& SetTitleBackgroundBrush(const FSlateBrush& InPaletteTitleBackgroundBrush);
 
 	UPROPERTY(EditAnywhere, BlueprintReadWrite, Category = Appearance)
 	FSlateBrush ToolDetailsBackgroundBrush;
-	FToolkitWidgetStyle& SetToolDetailsBackgroundBrush(const FSlateBrush& InToolDetailsBackgroundBrush);
+	UE_API FToolkitWidgetStyle& SetToolDetailsBackgroundBrush(const FSlateBrush& InToolDetailsBackgroundBrush);
 
 	UPROPERTY(EditAnywhere, BlueprintReadWrite, Category = Appearance)
 	FSlateColor TitleForegroundColor;
-	FToolkitWidgetStyle& SetTitleForegroundColor(const FSlateColor& InTitleForegroundColor);
+	UE_API FToolkitWidgetStyle& SetTitleForegroundColor(const FSlateColor& InTitleForegroundColor);
 	
 	UPROPERTY(EditAnywhere, BlueprintReadWrite, Category = Appearance)
 	FMargin TitlePadding;
-	FToolkitWidgetStyle& SetTitlePadding(const FMargin& InTitlePadding);
+	UE_API FToolkitWidgetStyle& SetTitlePadding(const FMargin& InTitlePadding);
 	
 	UPROPERTY(EditAnywhere, BlueprintReadWrite, Category = Appearance)
 	FMargin ActiveToolTitleBorderPadding;
-	FToolkitWidgetStyle& SetActiveToolTitleBorderPadding(const FMargin& InActiveToolTitleBorderPadding);
+	UE_API FToolkitWidgetStyle& SetActiveToolTitleBorderPadding(const FMargin& InActiveToolTitleBorderPadding);
 
 	UPROPERTY(EditAnywhere, BlueprintReadWrite, Category = Appearance)
 	FMargin ToolContextTextBlockPadding;
-	FToolkitWidgetStyle& SetToolContextTextBlockPadding(const FMargin& InToolContextTextBlockPadding);
+	UE_API FToolkitWidgetStyle& SetToolContextTextBlockPadding(const FMargin& InToolContextTextBlockPadding);
 
 	UPROPERTY(EditAnywhere, BlueprintReadWrite, Category = Appearance)
 	FSlateFontInfo TitleFont;
-	FToolkitWidgetStyle& SetTitleFont(const FSlateFontInfo& InTitleFont);
-};
\ No newline at end of file
+	UE_API FToolkitWidgetStyle& SetTitleFont(const FSlateFontInfo& InTitleFont);
+};
+
+#undef UE_API
diff --git a/Engine/Source/Developer/WidgetRegistration/Public/Layout/SeparatorBuilder.h b/Engine/Source/Developer/WidgetRegistration/Public/Layout/SeparatorBuilder.h
--- a/Engine/Source/Developer/WidgetRegistration/Public/Layout/SeparatorBuilder.h
+++ b/Engine/Source/Developer/WidgetRegistration/Public/Layout/SeparatorBuilder.h
@@ -113,7 +113,7 @@ public:
 	 * @param NewVisibility the TAttribute<EVisibility> to bind for the SSeparator. 
 	 * @return this FSeparatorBuilder to enable cascading.
 	 */
-	WIDGETREGISTRATION_API FSeparatorBuilder& BindVisibility(TAttribute<EVisibility> NewVisibility)
+	FSeparatorBuilder& BindVisibility(TAttribute<EVisibility> NewVisibility)
 	{
 		Visibility = NewVisibility;
 		return *this;
@@ -162,4 +162,4 @@ private:
 	 */
 	TAttribute<EVisibility> Visibility = EVisibility::Visible;
 
-};
\ No newline at end of file
+};
diff --git a/Engine/Source/Developer/WidgetRegistration/Public/Persistence/BuilderPersistenceManager.h b/Engine/Source/Developer/WidgetRegistration/Public/Persistence/BuilderPersistenceManager.h
--- a/Engine/Source/Developer/WidgetRegistration/Public/Persistence/BuilderPersistenceManager.h
+++ b/Engine/Source/Developer/WidgetRegistration/Public/Persistence/BuilderPersistenceManager.h
@@ -9,11 +9,13 @@
 
 #include "BuilderPersistenceManager.generated.h"
 
+#define UE_API WIDGETREGISTRATION_API
+
 /**
  * Manages an array of FNames to persist
  */
 USTRUCT()
-struct WIDGETREGISTRATION_API FPersistedNameArray
+struct FPersistedNameArray
 {
 	GENERATED_BODY()
 public:
@@ -26,7 +28,7 @@ public:
  * Manages a Bool to persist
  */
 USTRUCT()
-struct WIDGETREGISTRATION_API FPersistedBool
+struct FPersistedBool
 {
 	GENERATED_BODY()
 public:
@@ -39,8 +41,8 @@ public:
 /**
  * The Builder Persistence Manager handles persistence for Builders through use of FBuilderKeys
  */
-UCLASS(EditorConfig="BuilderPersistenceManager")
-class WIDGETREGISTRATION_API  UBuilderPersistenceManager : public UEditorConfigBase
+UCLASS(MinimalAPI, EditorConfig="BuilderPersistenceManager")
+class  UBuilderPersistenceManager : public UEditorConfigBase
 {
 	GENERATED_BODY()
 
@@ -50,7 +52,7 @@ public:
 	 * 
 	 * @param Key the FBuilderKey to retrieve the favorites for
 	 */
-	TArray<FName> GetFavoritesNames( const UE::DisplayBuilders::FBuilderKey& Key );
+	UE_API TArray<FName> GetFavoritesNames( const UE::DisplayBuilders::FBuilderKey& Key );
 
 	/**
 	 * Sets the default named favorites array for the Builder with FBuilderKey Key
@@ -58,7 +60,7 @@ public:
 	 * @param Key the FBuilderKey to retrieve the favorites for
 	 * @param Favorites the array of FNames for the favorites to persist
 	 */
-	void PersistFavoritesNames( const UE::DisplayBuilders::FBuilderKey& Key, TArray<FName>& Favorites );
+	UE_API void PersistFavoritesNames( const UE::DisplayBuilders::FBuilderKey& Key, TArray<FName>& Favorites );
 
 	/**
 	 * Sets the default button label EVisisbility for the Builder with FBuilderKey Key
@@ -66,7 +68,7 @@ public:
 	 * @param Key the FBuilderKey to retrieve the button label Bool for
 	 * @param Bool the bool for the button labels to persist
 	 */
-	void PersistShowButtonLabels( const UE::DisplayBuilders::FBuilderKey& Key, bool bValue );
+	UE_API void PersistShowButtonLabels( const UE::DisplayBuilders::FBuilderKey& Key, bool bValue );
 	
 	/**
 	 * @return the bool or the buttons labels for the Builder with FBuilderKey Key, if one exists, else an empty array is returned
@@ -74,17 +76,17 @@ public:
 	 * @param Key the FBuilderKey to retrieve the button label bool
 	 * @param PersistedBoolIfNoneFound the Bool to set as the new persistence value, if one is not found
 	 */
-	bool GetShowButtonLabels( const UE::DisplayBuilders::FBuilderKey& Key, bool bDefaultValue );
+	UE_API bool GetShowButtonLabels( const UE::DisplayBuilders::FBuilderKey& Key, bool bDefaultValue );
 
 	/**
 	 * Initialize the Persistence manager
 	 */
-	static void Initialize();
+	static UE_API void Initialize();
 	
 	/**
 	* Shuts down the Persistence manager
 	*/
-	static void ShutDown();
+	static UE_API void ShutDown();
 
 	/**
 	 * Gets the singleton for the Builder Persistence Manager
@@ -103,7 +105,7 @@ private:
 	UPROPERTY(meta=(EditorConfig))
 	TMap<FString, FPersistedBool> SavedNameToPersistedBoolMap;
 	
-	static TObjectPtr<UBuilderPersistenceManager> Instance;
+	static UE_API TObjectPtr<UBuilderPersistenceManager> Instance;
 
 private:
 	/**
@@ -112,7 +114,7 @@ private:
 	 * @param Key the FBuilderKey to retrieve the FNames array for
 	 * @param PersistenceKeySuffix the suffix to add to the FBuilderKey to persist the array of FNames with
 	 */
-	TArray<FName> GetPersistedArrayOfNames( const UE::DisplayBuilders::FBuilderKey& Key, FName PersistenceKeySuffix );
+	UE_API TArray<FName> GetPersistedArrayOfNames( const UE::DisplayBuilders::FBuilderKey& Key, FName PersistenceKeySuffix );
 
 	/**
 	 * Persists the array of FNames the Builder with FBuilderKey Key and suffix PersistenceKeySuffix
@@ -120,7 +122,7 @@ private:
 	 * @param Key the FBuilderKey to retrieve the favorites for
 	 * @param PersistenceKeySuffix the suffix added to the FBuilderKey to persist the array of FNames with
 	 */
-	void PersistArrayOfNames( const UE::DisplayBuilders::FBuilderKey& Key, FName PersistenceKeySuffix, TArray<FName>& ArrayOfNamesToPersist );
+	UE_API void PersistArrayOfNames( const UE::DisplayBuilders::FBuilderKey& Key, FName PersistenceKeySuffix, TArray<FName>& ArrayOfNamesToPersist );
 
 	/**
 	 * Returns the bool for the Builder with FBuilderKey Key and the suffix that the array was persisted with, if one exists, else an empty array is returned
@@ -129,7 +131,7 @@ private:
 	 * @param PersistenceKeySuffix the suffix to add to the FBuilderKey to get the bool with
 	 * @param bDefaultValue the Bool to set as the new persistence value, if one is not found
 	 */
-	bool GetPersistedBool( const UE::DisplayBuilders::FBuilderKey& Key, FName PersistenceKeySuffix, bool bDefaultValue );
+	UE_API bool GetPersistedBool( const UE::DisplayBuilders::FBuilderKey& Key, FName PersistenceKeySuffix, bool bDefaultValue );
 	
 	/**
 	 * Persists the Bool the Builder with FBuilderKey Key and suffix PersistenceKeySuffix
@@ -137,5 +139,7 @@ private:
 	 * @param Key the FBuilderKey to retrieve the favorites for
 	 * @param PersistenceKeySuffix the suffix added to the FBuilderKey to persist the Bool with
 	 */
-	void PersistBool( const UE::DisplayBuilders::FBuilderKey& Key, FName PersistenceKeySuffix, bool bInPersistedBool );
-};
\ No newline at end of file
+	UE_API void PersistBool( const UE::DisplayBuilders::FBuilderKey& Key, FName PersistenceKeySuffix, bool bInPersistedBool );
+};
+
+#undef UE_API
diff --git a/Engine/Source/Developer/WidgetRegistration/Public/Styles/SlateBrushTemplates.h b/Engine/Source/Developer/WidgetRegistration/Public/Styles/SlateBrushTemplates.h
--- a/Engine/Source/Developer/WidgetRegistration/Public/Styles/SlateBrushTemplates.h
+++ b/Engine/Source/Developer/WidgetRegistration/Public/Styles/SlateBrushTemplates.h
@@ -5,30 +5,34 @@
 #include "Styling/SlateBrush.h"
 #include "Styling/StyleColors.h"
 
+#define UE_API WIDGETREGISTRATION_API
+
 /** A const FSlateBrush* Factory */
-struct WIDGETREGISTRATION_API FSlateBrushTemplates
+struct FSlateBrushTemplates
 {
 	// Images
-	static const FSlateBrush* DragHandle();
-	static const FSlateBrush* ThinLineHorizontal();
+	static UE_API const FSlateBrush* DragHandle();
+	static UE_API const FSlateBrush* ThinLineHorizontal();
 
 	// Colors
-	static const FSlateBrush* Transparent();
-	static const FSlateBrush* Panel();
-	static const FSlateBrush* Recessed();
+	static UE_API const FSlateBrush* Transparent();
+	static UE_API const FSlateBrush* Panel();
+	static UE_API const FSlateBrush* Recessed();
 
 	/**
 	 * gets a const FSlateBrush* with the color Color
 	 *
 	 *  @param the EStyleColor we need the slate brush for
 	 */
-	const FSlateBrush* GetBrushWithColor(EStyleColor Color);
+	UE_API const FSlateBrush* GetBrushWithColor(EStyleColor Color);
 
 	/**
 	 * gets the FSlateBrushTemplates singleton
 	 */
-	static FSlateBrushTemplates& Get();
+	static UE_API FSlateBrushTemplates& Get();
 
 	/** the map of EStyleColor to const FSlateBrush */
 	TMap<EStyleColor, const FSlateBrush> EStyleColorToSlateBrushMap;
 };
+
+#undef UE_API
diff --git a/Engine/Source/Developer/WidgetRegistration/Public/ToolElementRegistry.h b/Engine/Source/Developer/WidgetRegistration/Public/ToolElementRegistry.h
--- a/Engine/Source/Developer/WidgetRegistration/Public/ToolElementRegistry.h
+++ b/Engine/Source/Developer/WidgetRegistration/Public/ToolElementRegistry.h
@@ -5,6 +5,8 @@
 #include "BuilderKey.h"
 #include "Widgets/SWidget.h"
 
+#define UE_API WIDGETREGISTRATION_API
+
 class FToolElementRegistry;
 
 /** An enum to provide the type of Tool Element*/
@@ -20,39 +22,39 @@ enum class EToolElement : uint32
  * extended and used to hold whatever information is needed to generate or update the
  * widgets that make up the tool, and implement the virtual widget methods
  */
-class WIDGETREGISTRATION_API FToolElementRegistrationArgs :
+class FToolElementRegistrationArgs :
 	public TSharedFromThis<FToolElementRegistrationArgs>
 {
 	
 public:
 
 	/** The constructor, which takes an EToolElement which will define the type of Tool Element*/
-	FToolElementRegistrationArgs(EToolElement InToolElementType);
+	UE_API FToolElementRegistrationArgs(EToolElement InToolElementType);
 
 	/** The constructor, which takes an EToolElement which will define the type of Tool Element*/
-	FToolElementRegistrationArgs(FName InStyleClassName);
+	UE_API FToolElementRegistrationArgs(FName InStyleClassName);
 	
 	/**
 	 * The constructor, which takes an FBuilderKey which will define the type of Tool Element
 	 *
 	 * @param InBuilderKey the FBuilderKey that provides the Key for this builder
 	 */
-	FToolElementRegistrationArgs( UE::DisplayBuilders::FBuilderKey InBuilderKey );
+	UE_API FToolElementRegistrationArgs( UE::DisplayBuilders::FBuilderKey InBuilderKey );
 
 	/** default destructor in case any subclasses need to provide a destructor*/
 	virtual ~FToolElementRegistrationArgs() = default;
 
 	/** Implements the generation of the TSharedPtr<SWidget> */
-	virtual TSharedPtr<SWidget> GenerateWidget();
+	UE_API virtual TSharedPtr<SWidget> GenerateWidget();
 
 	/** Implements the generation of the TSharedRef<SWidget> */
-	TSharedRef<SWidget> GenerateWidgetSharedRef();
+	UE_API TSharedRef<SWidget> GenerateWidgetSharedRef();
 
 	/** Updates/reloads this widget. This should be called after a consumer has changed any Data in this */
-	virtual void UpdateWidget();
+	UE_API virtual void UpdateWidget();
 
 	/** Resets the widget to its initial state */
-	virtual void ResetWidget();
+	UE_API virtual void ResetWidget();
 
 	/** The type of tool element this is Registration args for */
 	const EToolElement ToolElementType;
@@ -68,7 +70,7 @@ public:
 };
 
 /** Serves as a key into the FToolElementRegistry of FToolElements */
-class WIDGETREGISTRATION_API FToolElementRegistrationKey :
+class FToolElementRegistrationKey :
 	public TSharedFromThis<FToolElementRegistrationKey> {
 
 public:
@@ -78,7 +80,7 @@ public:
 	 * @param InName the FName that will be part of the compound key for the FToolElement this denotes
 	 * @param InToolElementType the EToolElementType that will be part of the compound key for the FToolElement this denotes
 	 */
-	FToolElementRegistrationKey(FName InName, EToolElement InToolElementType );
+	UE_API FToolElementRegistrationKey(FName InName, EToolElement InToolElementType );
 
 	/** the FName that will be part of the compound key for the FToolElement this denotes */
 	FName Name;
@@ -86,14 +88,14 @@ public:
 	/** the EToolElementType that will be part of the compound key for the FToolElement this denotes */ 
 	EToolElement ToolElementType;
 
-	FString GetKeyString();
+	UE_API FString GetKeyString();
 
 private:
 FString KeyString;
 };
 
 /** Represents one Tool Element, and is responsible for displaying the UI for that element  */
-class WIDGETREGISTRATION_API FToolElement : public FToolElementRegistrationKey
+class FToolElement : public FToolElementRegistrationKey
 {
 
 public:
@@ -107,7 +109,7 @@ public:
 	 * @param InBuilderArgs the shared reference to the FToolElementRegistrationArgs. Need to have a shared reference
 	 * here so that if the registered outlives the registrar, we maintain a copy
 	 */
-	FToolElement(
+	UE_API FToolElement(
 		const FName InName,
 		TSharedRef<FToolElementRegistrationArgs> InBuilderArgs);
 
@@ -117,14 +119,14 @@ public:
 	 *
 	 * @return the TSharedRef<SWidget> which makes up the UI for this tool element
 	 */
-	TSharedRef<SWidget> GenerateWidget();
+	UE_API TSharedRef<SWidget> GenerateWidget();
 
 	/*
 	 * Sets the FToolElementRegistrationArgs for the this tool element to RegistrationArgs
 	 *
 	 * @param RegistrationArgs the FToolElementRegistrationArgs for the this tool element
 	 */
-	void SetRegistrationArgs(TSharedRef<FToolElementRegistrationArgs> RegistrationArgs);
+	UE_API void SetRegistrationArgs(TSharedRef<FToolElementRegistrationArgs> RegistrationArgs);
 
 private:
 	/** the FToolElementRegistrationArgs for the this tool element */
@@ -133,7 +135,7 @@ private:
 };
 
 /** A Registry for FToolElements, where the FToolElements will generate the UI for a certain tool */
-class WIDGETREGISTRATION_API FToolElementRegistry 
+class FToolElementRegistry 
 {
 public:
 
@@ -142,10 +144,10 @@ public:
 	 *
 	 *@param ToolElementKey the given FToolElementRegistrationKey& tool element key for which to find the tool element
 	 */
-	TSharedPtr<FToolElement> GetToolElementSP( FToolElementRegistrationKey& ToolElementKey);
-	explicit FToolElementRegistry();
+	UE_API TSharedPtr<FToolElement> GetToolElementSP( FToolElementRegistrationKey& ToolElementKey);
+	UE_API explicit FToolElementRegistry();
 
-	static FToolElementRegistry& Get();
+	static UE_API FToolElementRegistry& Get();
 
 	/**
 	 * Generates the widget for the tool element denoted by the FToolElementRegistrationKey, using the specified
@@ -160,7 +162,7 @@ public:
 	 * @param bUpdateRegistrationArgs a bool that indicates whether we store any provided FToolElementRegistrationArgs
 	 * as the new FToolElementRegistrationArgs for this tool element
 	 */
-	TSharedRef< SWidget > GenerateWidget(
+	UE_API TSharedRef< SWidget > GenerateWidget(
 		TSharedRef<FToolElementRegistrationKey> ToolElementKeySR,
 		TSharedPtr<FToolElementRegistrationArgs> RegistrationArgsSP = nullptr,
 		bool bUpdateRegistrationArgs = false);
@@ -169,17 +171,19 @@ public:
 	 *
 	 * @param ToolElement the FToolElement that is to be registered
 	 */
-	void RegisterElement(const TSharedRef<FToolElement> ToolElement);
+	UE_API void RegisterElement(const TSharedRef<FToolElement> ToolElement);
 	
 	/** Unregisters the FToolElement ToolElement
 	 *
 	 * @param ToolElement the FToolElement that is to be registered
 	 */
-	void UnregisterElement(const TSharedRef<FToolElement> ToolElement);
+	UE_API void UnregisterElement(const TSharedRef<FToolElement> ToolElement);
 
 private:
 	/* the Map of tool element Key to tool element */
 	TMap<FString, TSharedPtr<FToolElement>> ToolElementKeyToToolElementMap;
 };
 
- 
\ No newline at end of file
+ 
+
+#undef UE_API
diff --git a/Engine/Source/Developer/WidgetRegistration/Public/ToolbarRegistrationArgs.h b/Engine/Source/Developer/WidgetRegistration/Public/ToolbarRegistrationArgs.h
--- a/Engine/Source/Developer/WidgetRegistration/Public/ToolbarRegistrationArgs.h
+++ b/Engine/Source/Developer/WidgetRegistration/Public/ToolbarRegistrationArgs.h
@@ -5,8 +5,10 @@
 #include "ToolElementRegistry.h"
 #include "Framework/MultiBox/MultiBoxBuilder.h"
 
+#define UE_API WIDGETREGISTRATION_API
+
 /** A class that provides the FToolbarRegistrationArgs for a UE toolbar  */
-class WIDGETREGISTRATION_API FToolbarRegistrationArgs : public FToolElementRegistrationArgs
+class FToolbarRegistrationArgs : public FToolElementRegistrationArgs
 {
 public:
 
@@ -16,16 +18,18 @@ public:
 	 * @param InToolBarBuilder a TSharedRef to the FToolBarBuilder  which will be used by
 	 * this to build up a toolbar SWidget
 	 */
-	FToolbarRegistrationArgs(TSharedRef<FToolBarBuilder> InToolBarBuilder);
+	UE_API FToolbarRegistrationArgs(TSharedRef<FToolBarBuilder> InToolBarBuilder);
 
 	/**
 	 * Generates the TSharedPtr<SWidget> for a toolbar
 	 *
 	 * @return the the TSharedPtr<SWidget> for a toolbar
 	 */
-	virtual TSharedPtr<SWidget> GenerateWidget() override;
+	UE_API virtual TSharedPtr<SWidget> GenerateWidget() override;
 
 	/** The FToolBarBuilder object used to build up a toolbar widget */
 	TSharedRef<FToolBarBuilder> ToolBarBuilder;
 	
 };
+
+#undef UE_API
diff --git a/Engine/Source/Developer/WidgetRegistration/Public/ToolkitBuilder.h b/Engine/Source/Developer/WidgetRegistration/Public/ToolkitBuilder.h
--- a/Engine/Source/Developer/WidgetRegistration/Public/ToolkitBuilder.h
+++ b/Engine/Source/Developer/WidgetRegistration/Public/ToolkitBuilder.h
@@ -12,10 +12,12 @@
 #include "Widgets/Layout/SSplitter.h"
 #include "Layout/CategoryDrivenContentBuilderBase.h"
 
+#define UE_API WIDGETREGISTRATION_API
+
 class SWidget;
 class FToolElement;
 
-struct WIDGETREGISTRATION_API FToolkitSections
+struct FToolkitSections
 {
 	TSharedPtr<STextBlock> ModeWarningArea = nullptr;
 	TSharedPtr<STextBlock> ToolWarningArea = nullptr;
@@ -25,7 +27,7 @@ struct WIDGETREGISTRATION_API FToolkitSections
 };
 
 /** A struct that provides the data for a single tool Palette*/
-struct WIDGETREGISTRATION_API FToolPalette : TSharedFromThis<FToolPalette>
+struct FToolPalette : TSharedFromThis<FToolPalette>
 {
 	FToolPalette(TSharedPtr<FUICommandInfo> InLoadToolPaletteAction,
 		const TArray<TSharedPtr< FUICommandInfo >>& InPaletteActions) :
@@ -50,9 +52,9 @@ struct WIDGETREGISTRATION_API FToolPalette : TSharedFromThis<FToolPalette>
 };
 
 /** An FToolPalette to which you can add and remove actions */
-struct WIDGETREGISTRATION_API FEditablePalette : public FToolPalette
+struct FEditablePalette : public FToolPalette
 {
-	FEditablePalette(TSharedPtr<FUICommandInfo> InLoadToolPaletteAction,
+	UE_API FEditablePalette(TSharedPtr<FUICommandInfo> InLoadToolPaletteAction,
 		TSharedPtr<FUICommandInfo> InAddToPaletteAction,
 		TSharedPtr<FUICommandInfo> InRemoveFromPaletteAction,
 		FName InEditablePaletteName = FName(),
@@ -72,19 +74,19 @@ struct WIDGETREGISTRATION_API FEditablePalette : public FToolPalette
 	 *
 	 * @param CommandName the name of the FUICommandInfo queried as to whether it is in the Palette
 	 */
-	bool IsInPalette(const FName CommandName) const;
+	UE_API bool IsInPalette(const FName CommandName) const;
 
-	TArray<FString> GetPaletteCommandNames() const;
+	UE_API TArray<FString> GetPaletteCommandNames() const;
 
-	void AddCommandToPalette(const FString CommandNameString);
+	UE_API void AddCommandToPalette(const FString CommandNameString);
 
-	void RemoveCommandFromPalette(const FString CommandNameString);
+	UE_API void RemoveCommandFromPalette(const FString CommandNameString);
 
 protected:
 
-	void SaveToConfig();
+	UE_API void SaveToConfig();
 
-	void LoadFromConfig();
+	UE_API void LoadFromConfig();
 	
 protected:
 	/** The TArray of Command names that are the current FuiCommandInfo actions in this Palette */
@@ -120,7 +122,7 @@ struct FToolkitBuilderArgs : public FCategoryDrivenContentBuilderArgs
 	EVisibility SelectedCategoryTitleVisibility;
 };
 
-class WIDGETREGISTRATION_API FToolkitBuilder : public FCategoryDrivenContentBuilderBase
+class FToolkitBuilder : public FCategoryDrivenContentBuilderBase
 {
 public:
 	/**
@@ -130,53 +132,53 @@ public:
 	 * @param InToolkitCommandList  the toolkit FUICommandList
 	 * @param InToolkitSections The FToolkitSections for this toolkit builder
 	 */
-	explicit FToolkitBuilder(
+	UE_API explicit FToolkitBuilder(
 	FName ToolbarCustomizationName,
 	TSharedPtr<FUICommandList> InToolkitCommandList,
 	TSharedPtr<FToolkitSections> InToolkitSections);
 
 	/** default constructor */
-	explicit FToolkitBuilder(FToolkitBuilderArgs& Args);
+	UE_API explicit FToolkitBuilder(FToolkitBuilderArgs& Args);
 
-	virtual ~FToolkitBuilder() override;
+	UE_API virtual ~FToolkitBuilder() override;
 
-	void SetActivePaletteCommandsVisibility(EVisibility Visibility);
+	UE_API void SetActivePaletteCommandsVisibility(EVisibility Visibility);
 
 	/**
 	 * @return the Visibility state of the active palette command buttons
 	 */
-	EVisibility GetActivePaletteCommandsVisibility() const;
+	UE_API EVisibility GetActivePaletteCommandsVisibility() const;
 
 	/**
 	 * Adds the FToolPalette Palette to this FToolkitBuilder
 	 *
 	 * @param Palette the FToolPalette being added to this FToolkitBuilder
 	 */
-	void AddPalette(TSharedPtr<FToolPalette> Palette);
+	UE_API void AddPalette(TSharedPtr<FToolPalette> Palette);
 
 	/**
 	 * Adds the FEditablePalette Palette to this FToolkitBuilder
 	 *
 	 * @param Palette the FEditablePalette being added to this FToolkitBuilder
 	 */
-	void AddPalette(TSharedPtr<FEditablePalette> Palette);
+	UE_API void AddPalette(TSharedPtr<FEditablePalette> Palette);
 
 	/** Cleans up any previously set data in this FToolkitBuilder and reset the members to their initial values  */
-	virtual void ResetWidget() override final;
+	UE_API virtual void ResetWidget() override final;
 
 	/** Updates the Toolkit. This should be called after any changes to the data of this FToolkitBuilder, and the UI
 	 * will be regenerated to reflect it.  */
-	virtual void UpdateWidget() override;
+	UE_API virtual void UpdateWidget() override;
 
 	/** @return true is there is an active palette selected, else it returns false */
-	bool HasActivePalette() const;
+	UE_API bool HasActivePalette() const;
 
 	/*
 	 * Loads the palette for the FUICommandInfo Command on first visit to the mode.
 	 *
 	 *  @param Command the FUICommandInfo which defines the palette that will be loaded on first visit to the mode.
 	 */
-	void SetActivePaletteOnLoad(const FUICommandInfo* Command);
+	UE_API void SetActivePaletteOnLoad(const FUICommandInfo* Command);
 	
 	/**
 	 * @return true if the FUICommandInfo with the name CommandName is the active tool palette,
@@ -184,12 +186,12 @@ public:
 	 *
 	 * @param CommandName the name of the FUICommandInfo we are checking to see if it is the active tool palette
 	 */
-	ECheckBoxState IsActiveToolPalette(FName CommandName) const;
+	UE_API ECheckBoxState IsActiveToolPalette(FName CommandName) const;
 	
 	/*
 	 * Initializes the data necessary to build the category toolbar
 	 */
-	void InitializeCategoryToolbar() override;
+	UE_API void InitializeCategoryToolbar() override;
 
 	/**
 	 * Initializes the data necessary to build the category toolbar
@@ -199,19 +201,19 @@ public:
 	 * to add new data when the method completes. Otherwise, the map will keep
 	 * the load action name to ToolPalette mappings that it was initially set up with. 
 	 */
-	void InitializeCategoryToolbar(bool bInitLoadToolPaletteMap);
+	UE_API void InitializeCategoryToolbar(bool bInitLoadToolPaletteMap);
 
 	/**
 	 * Sets the display name for the active tool
 	 *
 	 * @param InActiveToolDisplayName an FText that holds the name of the currently active tool
 	 */
-	void SetActiveToolDisplayName(FText InActiveToolDisplayName);
+	UE_API void SetActiveToolDisplayName(FText InActiveToolDisplayName);
 
 	/**
 	 * @returns the display name for the active tool
 	 */
-	FText GetActiveToolDisplayName() const;
+	UE_API FText GetActiveToolDisplayName() const;
 
 	/**
 	 * Fills the OutCommands TArray with TSharedPtr<const FUICommandInfo> instances that are present in the FEditablePalette
@@ -219,12 +221,12 @@ public:
 	 *
 	 * @param EditablePalette the FEditablePalette whose FUICommandInfo
 	 */
-	void GetCommandsForEditablePalette(TSharedRef<FEditablePalette> EditablePalette, TArray<TSharedPtr<const FUICommandInfo>>& OutCommands);
+	UE_API void GetCommandsForEditablePalette(TSharedRef<FEditablePalette> EditablePalette, TArray<TSharedPtr<const FUICommandInfo>>& OutCommands);
 
 	/**
 	 * @returns the name of the Active Palette if one is available, else it returns NAME_NONE
 	 */
-	FName GetActivePaletteName() const;
+	UE_API FName GetActivePaletteName() const;
 
 	/** Name of the toolbar this mode uses and can be used by external systems to customize that mode toolbar */
 	FName ToolbarCustomizationName;
@@ -254,12 +256,12 @@ public:
 	TMap<FName, TSharedPtr<FToolBarBuilder>> LoadCommandNameToPaletteToolbarBuilderMap;
 
 	/** Updates the Editable Palette with any commands that are on it */
-	void UpdateEditablePalette(TSharedRef<FEditablePalette> EditablePalette);
+	UE_API void UpdateEditablePalette(TSharedRef<FEditablePalette> EditablePalette);
 
-	void OnEditablePaletteEdited(TSharedRef<FEditablePalette> EditablePalette);
+	UE_API void OnEditablePaletteEdited(TSharedRef<FEditablePalette> EditablePalette);
 
 	/** Resets the tool palette widget, on which the buttons for the currently chosen toolset are shown */
-	void ResetToolPaletteWidget();
+	UE_API void ResetToolPaletteWidget();
 
 	/** The array of Tool Palettes on display. This array is used to keep a handle
 	* on all of the created FToolElements so that we can unregister them upon destruction */
@@ -271,7 +273,7 @@ private:
 	 *
 	 * @param Palette the FToolPalette for which this method will build the UI 
 	 */
-	void CreatePalette(TSharedPtr<FToolPalette> Palette);
+	UE_API void CreatePalette(TSharedPtr<FToolPalette> Palette);
 
 	/**
 	 * Creates the UI for the Palette specified by the FToolPalette Palette if it is not currently active,
@@ -279,7 +281,7 @@ private:
 	 *
 	 * @param Palette the FToolPalette for which this method will build the UI 
 	 */
-	void TogglePalette(TSharedPtr<FToolPalette> Palette);
+	UE_API void TogglePalette(TSharedPtr<FToolPalette> Palette);
 	
 	/**
 	 * Retrieves the context menu content for the FUICommandInfo with the command name CommandName
@@ -287,12 +289,12 @@ private:
 	 * @param CommandName the command name of the FUICommandInfo to get the context menu widget for
 	 * @return the TSharedRef<SWidget> which contains the context menu for the FUICommandInfo with the name CommandName 
 	 */
-	TSharedRef<SWidget> GetContextMenuContent(const FName CommandName);
+	UE_API TSharedRef<SWidget> GetContextMenuContent(const FName CommandName);
 	
 	/** @return a TSharedPointer to the FToolbarBuilder with the FUICommandInfos that load the Palettes */
-	TSharedRef<SWidget> GetToolPaletteWidget() const;
+	UE_API TSharedRef<SWidget> GetToolPaletteWidget() const;
 	
-	void CreatePaletteWidget(FToolPalette& Palette, FToolElement& Element);
+	UE_API void CreatePaletteWidget(FToolPalette& Palette, FToolElement& Element);
 	
 	/**
 	 * Adds the FUICommandInfo with the command name CommandNameString to the FEditablePalette Palette if
@@ -301,24 +303,24 @@ private:
 	 * @param Palette the FEditablePalette to which the FUICommandInfo will be toggled on/off
 	 * @param CommandNameString the name of the FUICommandInfo command which will be toggled on/off of Palette
 	 */
-	void ToggleCommandInPalette(TSharedRef<FEditablePalette> Palette, FString CommandNameString);
+	UE_API void ToggleCommandInPalette(TSharedRef<FEditablePalette> Palette, FString CommandNameString);
 		
 	/* @return true if the Toolkit builder has some tools that are currently active/selected */
-	bool HasSelectedToolSet() const;
+	UE_API bool HasSelectedToolSet() const;
 
 	/**
 	 * Given the active category name, update the toolkit content
 	 * 
 	 * @param ActiveCategoryName 
 	 */
-	virtual void UpdateContentForCategory( FName ActiveCategoryName = NAME_None, FText InActiveCategoryText = FText::GetEmpty() ) override;
+	UE_API virtual void UpdateContentForCategory( FName ActiveCategoryName = NAME_None, FText InActiveCategoryText = FText::GetEmpty() ) override;
 
 	/**
 	 * Gets the EVisibility of the active tool title. This should only be visible if a tool is currently chosen in the palette.
 	 *
 	 * @return the EVisibility of the active tool title
 	 */
-	EVisibility GetActiveToolTitleVisibility() const;
+	UE_API EVisibility GetActiveToolTitleVisibility() const;
 	
 	/** If SelectedCategoryTitleVisibility == EVisibility::Visible, the selected
     * category title is visible, else it is not displayed. By default the selected category title is Collapsed  */	
@@ -336,4 +338,6 @@ private:
 
 	/** The array of Command names, in the correct display order */
 	TArray<FName> LoadCommandArray;
-};
\ No newline at end of file
+};
+
+#undef UE_API
diff --git a/Engine/Source/Developer/WidgetRegistration/Public/ToolkitBuilderConfig.h b/Engine/Source/Developer/WidgetRegistration/Public/ToolkitBuilderConfig.h
--- a/Engine/Source/Developer/WidgetRegistration/Public/ToolkitBuilderConfig.h
+++ b/Engine/Source/Developer/WidgetRegistration/Public/ToolkitBuilderConfig.h
@@ -8,7 +8,7 @@
 
 // The actual struct that contains the properties saved for an FEditableToolPalette
 USTRUCT()
-struct WIDGETREGISTRATION_API FEditableToolPaletteSettings
+struct FEditableToolPaletteSettings
 {
 	GENERATED_BODY()
 	
diff --git a/Engine/Source/Developer/WidgetRegistration/Public/ToolkitStyle.h b/Engine/Source/Developer/WidgetRegistration/Public/ToolkitStyle.h
--- a/Engine/Source/Developer/WidgetRegistration/Public/ToolkitStyle.h
+++ b/Engine/Source/Developer/WidgetRegistration/Public/ToolkitStyle.h
@@ -5,24 +5,27 @@
 #include "Styling/SlateStyle.h"
 #include "Templates/SharedPointer.h"
 
+#define UE_API WIDGETREGISTRATION_API
+
 class ISlateStyle;
 
 /** FToolkitStyle is the FSlateStyleSet that defines styles for FToolkitBuilders  */
-class WIDGETREGISTRATION_API FToolkitStyle
+class FToolkitStyle
 	: public FSlateStyleSet
 {
 public:
 	
-	static FName StyleName;
-	FToolkitStyle();
-	static void Initialize();
-	static void Shutdown();
-	static const ISlateStyle& Get();
-	virtual const FName& GetStyleSetName() const override;
+	static UE_API FName StyleName;
+	UE_API FToolkitStyle();
+	static UE_API void Initialize();
+	static UE_API void Shutdown();
+	static UE_API const ISlateStyle& Get();
+	UE_API virtual const FName& GetStyleSetName() const override;
 
 private:
 
-	static TSharedRef< class FSlateStyleSet > Create();
-	static TSharedPtr< class FSlateStyleSet > StyleSet;
+	static UE_API TSharedRef< class FSlateStyleSet > Create();
+	static UE_API TSharedPtr< class FSlateStyleSet > StyleSet;
 };
 
+#undef UE_API
diff --git a/Engine/Source/Editor/AIGraph/Classes/AIGraph.h b/Engine/Source/Editor/AIGraph/Classes/AIGraph.h
--- a/Engine/Source/Editor/AIGraph/Classes/AIGraph.h
+++ b/Engine/Source/Editor/AIGraph/Classes/AIGraph.h
@@ -11,45 +11,47 @@
 
 #include "AIGraph.generated.h"
 
+#define UE_API AIGRAPH_API
+
 class FArchive;
 class UEdGraphPin;
 class UObject;
 
-UCLASS()
-class AIGRAPH_API UAIGraph : public UEdGraph
+UCLASS(MinimalAPI)
+class UAIGraph : public UEdGraph
 {
 	GENERATED_UCLASS_BODY()
 
 	UPROPERTY()
 	int32 GraphVersion;
 
-	virtual void OnCreated();
-	virtual void OnLoaded();
-	virtual void Initialize();
+	UE_API virtual void OnCreated();
+	UE_API virtual void OnLoaded();
+	UE_API virtual void Initialize();
 
-	virtual void UpdateAsset(int32 UpdateFlags = 0);
-	virtual void UpdateVersion();
-	virtual void MarkVersion();
+	UE_API virtual void UpdateAsset(int32 UpdateFlags = 0);
+	UE_API virtual void UpdateVersion();
+	UE_API virtual void MarkVersion();
 
-	virtual void OnSubNodeDropped();
-	virtual void OnNodesPasted(const FString& ImportStr);
+	UE_API virtual void OnSubNodeDropped();
+	UE_API virtual void OnNodesPasted(const FString& ImportStr);
 
-	bool UpdateUnknownNodeClasses();
+	UE_API bool UpdateUnknownNodeClasses();
 
 	UE_DEPRECATED(5.6, "This method was only used to fetch error messages and log them. It has been renamed to better fit its purpose. Use UpdateErrorMessages instead.")
-	void UpdateDeprecatedClasses();
+	UE_API void UpdateDeprecatedClasses();
 
 	/** Sets error messages and logs errors about nodes. */
-	void UpdateErrorMessages();
-	void RemoveOrphanedNodes();
-	void UpdateClassData();
+	UE_API void UpdateErrorMessages();
+	UE_API void RemoveOrphanedNodes();
+	UE_API void UpdateClassData();
 
-	bool IsLocked() const;
-	void LockUpdates();
-	void UnlockUpdates();
+	UE_API bool IsLocked() const;
+	UE_API void LockUpdates();
+	UE_API void UnlockUpdates();
 
 	//~ Begin UObject Interface.
-	virtual void Serialize(FArchive& Ar) override;
+	UE_API virtual void Serialize(FArchive& Ar) override;
 	//~ End UObject Interface.
 
 protected:
@@ -59,9 +61,11 @@ protected:
 	 */
 	uint32 bLockUpdates : 1;
 
-	virtual void CollectAllNodeInstances(TSet<UObject*>& NodeInstances);
-	virtual bool CanRemoveNestedObject(UObject* TestObject) const;
-	virtual void OnNodeInstanceRemoved(UObject* NodeInstance);
+	UE_API virtual void CollectAllNodeInstances(TSet<UObject*>& NodeInstances);
+	UE_API virtual bool CanRemoveNestedObject(UObject* TestObject) const;
+	UE_API virtual void OnNodeInstanceRemoved(UObject* NodeInstance);
 
-	UEdGraphPin* FindGraphNodePin(UEdGraphNode* Node, EEdGraphPinDirection Dir);
+	UE_API UEdGraphPin* FindGraphNodePin(UEdGraphNode* Node, EEdGraphPinDirection Dir);
 };
+
+#undef UE_API
diff --git a/Engine/Source/Editor/AIGraph/Classes/AIGraphNode.h b/Engine/Source/Editor/AIGraph/Classes/AIGraphNode.h
--- a/Engine/Source/Editor/AIGraph/Classes/AIGraphNode.h
+++ b/Engine/Source/Editor/AIGraph/Classes/AIGraphNode.h
@@ -8,14 +8,16 @@
 #include "AIGraphTypes.h"
 #include "AIGraphNode.generated.h"
 
+#define UE_API AIGRAPH_API
+
 class UEdGraph;
 class UEdGraphPin;
 class UEdGraphSchema;
 struct FDiffResults;
 struct FDiffSingleResult;
 
-UCLASS()
-class AIGRAPH_API UAIGraphNode : public UEdGraphNode
+UCLASS(MinimalAPI)
+class UAIGraphNode : public UEdGraphNode
 {
 	GENERATED_UCLASS_BODY()
 
@@ -49,72 +51,74 @@ class AIGRAPH_API UAIGraphNode : public UEdGraphNode
 	FString ErrorMessage;
 
 	//~ Begin UEdGraphNode Interface
-	virtual class UAIGraph* GetAIGraph();
-	virtual void AutowireNewNode(UEdGraphPin* FromPin) override;
-	virtual void PostPlacedNewNode() override;
-	virtual void PrepareForCopying() override;
-	virtual bool CanDuplicateNode() const override;
-	virtual bool CanUserDeleteNode() const override;
-	virtual void DestroyNode() override;
-	virtual FText GetTooltipText() const override;
-	virtual void NodeConnectionListChanged() override;
-	virtual bool CanCreateUnderSpecifiedSchema(const UEdGraphSchema* DesiredSchema) const override;
-	virtual void FindDiffs(class UEdGraphNode* OtherNode, struct FDiffResults& Results) override;
-	virtual FString GetPropertyNameAndValueForDiff(const FProperty* Prop, const uint8* PropertyAddr) const override;
+	UE_API virtual class UAIGraph* GetAIGraph();
+	UE_API virtual void AutowireNewNode(UEdGraphPin* FromPin) override;
+	UE_API virtual void PostPlacedNewNode() override;
+	UE_API virtual void PrepareForCopying() override;
+	UE_API virtual bool CanDuplicateNode() const override;
+	UE_API virtual bool CanUserDeleteNode() const override;
+	UE_API virtual void DestroyNode() override;
+	UE_API virtual FText GetTooltipText() const override;
+	UE_API virtual void NodeConnectionListChanged() override;
+	UE_API virtual bool CanCreateUnderSpecifiedSchema(const UEdGraphSchema* DesiredSchema) const override;
+	UE_API virtual void FindDiffs(class UEdGraphNode* OtherNode, struct FDiffResults& Results) override;
+	UE_API virtual FString GetPropertyNameAndValueForDiff(const FProperty* Prop, const uint8* PropertyAddr) const override;
 	//~ End UEdGraphNode Interface
 
 	//~ Begin UObject Interface
 #if WITH_EDITOR
-	virtual void PostEditImport() override;
-	virtual void PostEditUndo() override;
+	UE_API virtual void PostEditImport() override;
+	UE_API virtual void PostEditUndo() override;
 #endif
 	// End UObject
 
 	// @return the input pin for this state
-	virtual UEdGraphPin* GetInputPin(int32 InputIndex = 0) const;
+	UE_API virtual UEdGraphPin* GetInputPin(int32 InputIndex = 0) const;
 	// @return the output pin for this state
-	virtual UEdGraphPin* GetOutputPin(int32 InputIndex = 0) const;
+	UE_API virtual UEdGraphPin* GetOutputPin(int32 InputIndex = 0) const;
 	virtual UEdGraph* GetBoundGraph() const { return NULL; }
 
-	virtual FText GetDescription() const;
-	virtual void PostCopyNode();
+	UE_API virtual FText GetDescription() const;
+	UE_API virtual void PostCopyNode();
 
-	void AddSubNode(UAIGraphNode* SubNode, class UEdGraph* ParentGraph);
-	void RemoveSubNode(UAIGraphNode* SubNode);
-	virtual void RemoveAllSubNodes();
-	virtual void OnSubNodeRemoved(UAIGraphNode* SubNode);
-	virtual void OnSubNodeAdded(UAIGraphNode* SubNode);
+	UE_API void AddSubNode(UAIGraphNode* SubNode, class UEdGraph* ParentGraph);
+	UE_API void RemoveSubNode(UAIGraphNode* SubNode);
+	UE_API virtual void RemoveAllSubNodes();
+	UE_API virtual void OnSubNodeRemoved(UAIGraphNode* SubNode);
+	UE_API virtual void OnSubNodeAdded(UAIGraphNode* SubNode);
 
-	virtual int32 FindSubNodeDropIndex(UAIGraphNode* SubNode) const;
-	virtual void InsertSubNodeAt(UAIGraphNode* SubNode, int32 DropIndex);
+	UE_API virtual int32 FindSubNodeDropIndex(UAIGraphNode* SubNode) const;
+	UE_API virtual void InsertSubNodeAt(UAIGraphNode* SubNode, int32 DropIndex);
 
 	/** check if node is subnode */
-	virtual bool IsSubNode() const;
+	UE_API virtual bool IsSubNode() const;
 
 	/** initialize instance object  */
-	virtual void InitializeInstance();
+	UE_API virtual void InitializeInstance();
 
 	/** reinitialize node instance */
-	virtual bool RefreshNodeClass();
+	UE_API virtual bool RefreshNodeClass();
 
 	/** updates ClassData from node instance */
-	virtual void UpdateNodeClassData();
+	UE_API virtual void UpdateNodeClassData();
 
 	/**
 	 * Checks for any errors in this node and updates ErrorMessage with any resulting message
 	 * Called every time the graph is serialized (i.e. loaded, saved, execution index changed, etc)
 	 */
-	virtual void UpdateErrorMessage();
+	UE_API virtual void UpdateErrorMessage();
 
 	/** Check if node instance uses blueprint for its implementation */
-	bool UsesBlueprint() const;
+	UE_API bool UsesBlueprint() const;
 
 	/** check if node has any errors, used for assigning colors on graph */
-	virtual bool HasErrors() const;
+	UE_API virtual bool HasErrors() const;
 
-	static void UpdateNodeClassDataFrom(UClass* InstanceClass, FGraphNodeClassData& UpdatedData);
+	static UE_API void UpdateNodeClassDataFrom(UClass* InstanceClass, FGraphNodeClassData& UpdatedData);
 
 protected:
 
-	virtual void ResetNodeOwner();
+	UE_API virtual void ResetNodeOwner();
 };
+
+#undef UE_API
diff --git a/Engine/Source/Editor/AIGraph/Classes/AIGraphSchema.h b/Engine/Source/Editor/AIGraph/Classes/AIGraphSchema.h
--- a/Engine/Source/Editor/AIGraph/Classes/AIGraphSchema.h
+++ b/Engine/Source/Editor/AIGraph/Classes/AIGraphSchema.h
@@ -8,12 +8,14 @@
 #include "AIGraphTypes.h"
 #include "AIGraphSchema.generated.h"
 
+#define UE_API AIGRAPH_API
+
 class FSlateRect;
 class UEdGraph;
 
 /** Action to add a comment to the graph */
 USTRUCT()
-struct AIGRAPH_API FAISchemaAction_AddComment : public FEdGraphSchemaAction
+struct FAISchemaAction_AddComment : public FEdGraphSchemaAction
 {
 	GENERATED_BODY()
 	
@@ -24,13 +26,13 @@ struct AIGRAPH_API FAISchemaAction_AddComment : public FEdGraphSchemaAction
 	}
 
 	// FEdGraphSchemaAction interface
-	virtual UEdGraphNode* PerformAction(class UEdGraph* ParentGraph, UEdGraphPin* FromPin, const FVector2f& Location, bool bSelectNewNode = true) override final;
+	UE_API virtual UEdGraphNode* PerformAction(class UEdGraph* ParentGraph, UEdGraphPin* FromPin, const FVector2f& Location, bool bSelectNewNode = true) override final;
 	// End of FEdGraphSchemaAction interface
 };
 
 /** Action to add a node to the graph */
 USTRUCT()
-struct AIGRAPH_API FAISchemaAction_NewNode : public FEdGraphSchemaAction
+struct FAISchemaAction_NewNode : public FEdGraphSchemaAction
 {
 	GENERATED_USTRUCT_BODY();
 
@@ -49,15 +51,15 @@ struct AIGRAPH_API FAISchemaAction_NewNode : public FEdGraphSchemaAction
 	{}
 
 	//~ Begin FEdGraphSchemaAction Interface
-	virtual UEdGraphNode* PerformAction(class UEdGraph* ParentGraph, UEdGraphPin* FromPin, const FVector2f& Location, bool bSelectNewNode = true) override;
-	virtual UEdGraphNode* PerformAction(class UEdGraph* ParentGraph, TArray<UEdGraphPin*>& FromPins, const FVector2f& Location, bool bSelectNewNode = true) override;
-	virtual void AddReferencedObjects(FReferenceCollector& Collector) override;
+	UE_API virtual UEdGraphNode* PerformAction(class UEdGraph* ParentGraph, UEdGraphPin* FromPin, const FVector2f& Location, bool bSelectNewNode = true) override;
+	UE_API virtual UEdGraphNode* PerformAction(class UEdGraph* ParentGraph, TArray<UEdGraphPin*>& FromPins, const FVector2f& Location, bool bSelectNewNode = true) override;
+	UE_API virtual void AddReferencedObjects(FReferenceCollector& Collector) override;
 	//~ End FEdGraphSchemaAction Interface
 };
 
 /** Action to add a subnode to the selected node */
 USTRUCT()
-struct AIGRAPH_API FAISchemaAction_NewSubNode : public FEdGraphSchemaAction
+struct FAISchemaAction_NewSubNode : public FEdGraphSchemaAction
 {
 	GENERATED_USTRUCT_BODY();
 
@@ -82,35 +84,36 @@ struct AIGRAPH_API FAISchemaAction_NewSubNode : public FEdGraphSchemaAction
 	{}
 
 	//~ Begin FEdGraphSchemaAction Interface
-	virtual UEdGraphNode* PerformAction(class UEdGraph* ParentGraph, UEdGraphPin* FromPin, const FVector2f& Location, bool bSelectNewNode = true) override;
-	virtual UEdGraphNode* PerformAction(class UEdGraph* ParentGraph, TArray<UEdGraphPin*>& FromPins, const FVector2f& Location, bool bSelectNewNode = true) override;
-	virtual void AddReferencedObjects(FReferenceCollector& Collector) override;
+	UE_API virtual UEdGraphNode* PerformAction(class UEdGraph* ParentGraph, UEdGraphPin* FromPin, const FVector2f& Location, bool bSelectNewNode = true) override;
+	UE_API virtual UEdGraphNode* PerformAction(class UEdGraph* ParentGraph, TArray<UEdGraphPin*>& FromPins, const FVector2f& Location, bool bSelectNewNode = true) override;
+	UE_API virtual void AddReferencedObjects(FReferenceCollector& Collector) override;
 	//~ End FEdGraphSchemaAction Interface
 };
 
-UCLASS()
-class AIGRAPH_API UAIGraphSchema : public UEdGraphSchema
+UCLASS(MinimalAPI)
+class UAIGraphSchema : public UEdGraphSchema
 {
 	GENERATED_UCLASS_BODY()
 
 	//~ Begin EdGraphSchema Interface
-	virtual void GetContextMenuActions(class UToolMenu* Menu, class UGraphNodeContextMenuContext* Context) const override;
-	virtual FLinearColor GetPinTypeColor(const FEdGraphPinType& PinType) const override;
-	virtual bool ShouldHidePinDefaultValue(UEdGraphPin* Pin) const override;
-	virtual class FConnectionDrawingPolicy* CreateConnectionDrawingPolicy(int32 InBackLayerID, int32 InFrontLayerID, float InZoomFactor, const FSlateRect& InClippingRect, class FSlateWindowElementList& InDrawElements, class UEdGraph* InGraphObj) const override;
-	virtual void BreakNodeLinks(UEdGraphNode& TargetNode) const override;
-	virtual void BreakPinLinks(UEdGraphPin& TargetPin, bool bSendsNodeNotification) const override;
-	virtual void BreakSinglePinLink(UEdGraphPin* SourcePin, UEdGraphPin* TargetPin) const override;
-	virtual TSharedPtr<FEdGraphSchemaAction> GetCreateCommentAction() const override;
-	virtual int32 GetNodeSelectionCount(const UEdGraph* Graph) const override;
+	UE_API virtual void GetContextMenuActions(class UToolMenu* Menu, class UGraphNodeContextMenuContext* Context) const override;
+	UE_API virtual FLinearColor GetPinTypeColor(const FEdGraphPinType& PinType) const override;
+	UE_API virtual bool ShouldHidePinDefaultValue(UEdGraphPin* Pin) const override;
+	UE_API virtual class FConnectionDrawingPolicy* CreateConnectionDrawingPolicy(int32 InBackLayerID, int32 InFrontLayerID, float InZoomFactor, const FSlateRect& InClippingRect, class FSlateWindowElementList& InDrawElements, class UEdGraph* InGraphObj) const override;
+	UE_API virtual void BreakNodeLinks(UEdGraphNode& TargetNode) const override;
+	UE_API virtual void BreakPinLinks(UEdGraphPin& TargetPin, bool bSendsNodeNotification) const override;
+	UE_API virtual void BreakSinglePinLink(UEdGraphPin* SourcePin, UEdGraphPin* TargetPin) const override;
+	UE_API virtual TSharedPtr<FEdGraphSchemaAction> GetCreateCommentAction() const override;
+	UE_API virtual int32 GetNodeSelectionCount(const UEdGraph* Graph) const override;
 	//~ End EdGraphSchema Interface
 
-	virtual void GetGraphNodeContextActions(FGraphContextMenuBuilder& ContextMenuBuilder, int32 SubNodeFlags) const;
-	virtual void GetSubNodeClasses(int32 SubNodeFlags, TArray<FGraphNodeClassData>& ClassData, UClass*& GraphNodeClass) const;
+	UE_API virtual void GetGraphNodeContextActions(FGraphContextMenuBuilder& ContextMenuBuilder, int32 SubNodeFlags) const;
+	UE_API virtual void GetSubNodeClasses(int32 SubNodeFlags, TArray<FGraphNodeClassData>& ClassData, UClass*& GraphNodeClass) const;
 
 protected:
 
-	static TSharedPtr<FAISchemaAction_NewNode> AddNewNodeAction(FGraphActionListBuilderBase& ContextMenuBuilder, const FText& Category, const FText& MenuDesc, const FText& Tooltip);
-	static TSharedPtr<FAISchemaAction_NewSubNode> AddNewSubNodeAction(FGraphActionListBuilderBase& ContextMenuBuilder, const FText& Category, const FText& MenuDesc, const FText& Tooltip);
+	static UE_API TSharedPtr<FAISchemaAction_NewNode> AddNewNodeAction(FGraphActionListBuilderBase& ContextMenuBuilder, const FText& Category, const FText& MenuDesc, const FText& Tooltip);
+	static UE_API TSharedPtr<FAISchemaAction_NewSubNode> AddNewSubNodeAction(FGraphActionListBuilderBase& ContextMenuBuilder, const FText& Category, const FText& MenuDesc, const FText& Tooltip);
 };
 
+#undef UE_API
diff --git a/Engine/Source/Editor/AIGraph/Classes/AIGraphTypes.h b/Engine/Source/Editor/AIGraph/Classes/AIGraphTypes.h
--- a/Engine/Source/Editor/AIGraph/Classes/AIGraphTypes.h
+++ b/Engine/Source/Editor/AIGraph/Classes/AIGraphTypes.h
@@ -7,23 +7,25 @@
 #include "UObject/WeakObjectPtr.h"
 #include "AIGraphTypes.generated.h"
 
+#define UE_API AIGRAPH_API
+
 USTRUCT()
-struct AIGRAPH_API FGraphNodeClassData
+struct FGraphNodeClassData
 {
 	GENERATED_USTRUCT_BODY()
 
 	FGraphNodeClassData() {}
-	FGraphNodeClassData(UClass* InClass, const FString& InDeprecatedMessage);
-	FGraphNodeClassData(const FTopLevelAssetPath& InGeneratedClassPath, UClass* InClass);
-	FGraphNodeClassData(const FString& InAssetName, const FString& InGeneratedClassPackage, const FString& InClassName, UClass* InClass);
-
-	FString ToString() const;
-	FString GetClassName() const;
-	FText GetCategory() const;
-	FString GetDisplayName() const;
-	FText GetTooltip() const;
-	UClass* GetClass(bool bSilent = false);
-	bool IsAbstract() const;
+	UE_API FGraphNodeClassData(UClass* InClass, const FString& InDeprecatedMessage);
+	UE_API FGraphNodeClassData(const FTopLevelAssetPath& InGeneratedClassPath, UClass* InClass);
+	UE_API FGraphNodeClassData(const FString& InAssetName, const FString& InGeneratedClassPackage, const FString& InClassName, UClass* InClass);
+
+	UE_API FString ToString() const;
+	UE_API FString GetClassName() const;
+	UE_API FText GetCategory() const;
+	UE_API FString GetDisplayName() const;
+	UE_API FText GetTooltip() const;
+	UE_API UClass* GetClass(bool bSilent = false);
+	UE_API bool IsAbstract() const;
 
 	FORCEINLINE bool IsBlueprint() const { return AssetName.Len() > 0; }
 	FORCEINLINE bool IsDeprecated() const { return DeprecatedMessage.Len() > 0; }
@@ -60,7 +62,7 @@ private:
 	FString DeprecatedMessage;
 };
 
-struct AIGRAPH_API FGraphNodeClassNode
+struct FGraphNodeClassNode
 {
 	FGraphNodeClassData Data;
 	FString ParentClassName;
@@ -68,57 +70,59 @@ struct AIGRAPH_API FGraphNodeClassNode
 	TSharedPtr<FGraphNodeClassNode> ParentNode;
 	TArray<TSharedPtr<FGraphNodeClassNode> > SubNodes;
 
-	void AddUniqueSubNode(TSharedPtr<FGraphNodeClassNode> SubNode);
+	UE_API void AddUniqueSubNode(TSharedPtr<FGraphNodeClassNode> SubNode);
 };
 
-struct AIGRAPH_API FGraphNodeClassHelper
+struct FGraphNodeClassHelper
 {
 	DECLARE_MULTICAST_DELEGATE(FOnPackageListUpdated);
 
-	FGraphNodeClassHelper(UClass* InRootClass);
-	~FGraphNodeClassHelper();
+	UE_API FGraphNodeClassHelper(UClass* InRootClass);
+	UE_API ~FGraphNodeClassHelper();
 
-	void GatherClasses(const UClass* BaseClass, TArray<FGraphNodeClassData>& AvailableClasses);
-	static FString GetDeprecationMessage(const UClass* Class);
+	UE_API void GatherClasses(const UClass* BaseClass, TArray<FGraphNodeClassData>& AvailableClasses);
+	static UE_API FString GetDeprecationMessage(const UClass* Class);
 
-	void OnAssetAdded(const struct FAssetData& AssetData);
-	void OnAssetRemoved(const struct FAssetData& AssetData);
-	void InvalidateCache();
-	void OnReloadComplete(EReloadCompleteReason Reason);
+	UE_API void OnAssetAdded(const struct FAssetData& AssetData);
+	UE_API void OnAssetRemoved(const struct FAssetData& AssetData);
+	UE_API void InvalidateCache();
+	UE_API void OnReloadComplete(EReloadCompleteReason Reason);
 
-	static void AddUnknownClass(const FGraphNodeClassData& ClassData);
-	static bool IsClassKnown(const FGraphNodeClassData& ClassData);
-	static FOnPackageListUpdated OnPackageListUpdated;
+	static UE_API void AddUnknownClass(const FGraphNodeClassData& ClassData);
+	static UE_API bool IsClassKnown(const FGraphNodeClassData& ClassData);
+	static UE_API FOnPackageListUpdated OnPackageListUpdated;
 
-	static int32 GetObservedBlueprintClassCount(UClass* BaseNativeClass);
-	static void AddObservedBlueprintClasses(UClass* BaseNativeClass);
-	void UpdateAvailableBlueprintClasses();
+	static UE_API int32 GetObservedBlueprintClassCount(UClass* BaseNativeClass);
+	static UE_API void AddObservedBlueprintClasses(UClass* BaseNativeClass);
+	UE_API void UpdateAvailableBlueprintClasses();
 
 	/** Adds a single class to the list of hidden classes */
-	void AddForcedHiddenClass(UClass* Class);
+	UE_API void AddForcedHiddenClass(UClass* Class);
 
 	/** Overrides all previously set hidden classes */
-	void SetForcedHiddenClasses(const TSet<UClass*>& Classes);
+	UE_API void SetForcedHiddenClasses(const TSet<UClass*>& Classes);
 
-	void SetGatherBlueprints(bool bGather);
+	UE_API void SetGatherBlueprints(bool bGather);
 
 private:
 
 	UClass* RootNodeClass;
 	TSharedPtr<FGraphNodeClassNode> RootNode;
-	static TArray<FName> UnknownPackages;
-	static TMap<UClass*, int32> BlueprintClassCount;
+	static UE_API TArray<FName> UnknownPackages;
+	static UE_API TMap<UClass*, int32> BlueprintClassCount;
 	TSet<UClass*> ForcedHiddenClasses;
 	bool bGatherBlueprints = true;
 
-	TSharedPtr<FGraphNodeClassNode> CreateClassDataNode(const struct FAssetData& AssetData);
-	TSharedPtr<FGraphNodeClassNode> FindBaseClassNode(TSharedPtr<FGraphNodeClassNode> Node, const FString& ClassName);
-	void FindAllSubClasses(TSharedPtr<FGraphNodeClassNode> Node, TArray<FGraphNodeClassData>& AvailableClasses);
+	UE_API TSharedPtr<FGraphNodeClassNode> CreateClassDataNode(const struct FAssetData& AssetData);
+	UE_API TSharedPtr<FGraphNodeClassNode> FindBaseClassNode(TSharedPtr<FGraphNodeClassNode> Node, const FString& ClassName);
+	UE_API void FindAllSubClasses(TSharedPtr<FGraphNodeClassNode> Node, TArray<FGraphNodeClassData>& AvailableClasses);
 
-	void BuildClassGraph();
-	void AddClassGraphChildren(TSharedPtr<FGraphNodeClassNode> Node, TArray<TSharedPtr<FGraphNodeClassNode> >& NodeList);
+	UE_API void BuildClassGraph();
+	UE_API void AddClassGraphChildren(TSharedPtr<FGraphNodeClassNode> Node, TArray<TSharedPtr<FGraphNodeClassNode> >& NodeList);
 
-	bool IsHidingClass(UClass* Class);
-	bool IsHidingParentClass(UClass* Class);
-	bool IsPackageSaved(FName PackageName);
+	UE_API bool IsHidingClass(UClass* Class);
+	UE_API bool IsHidingParentClass(UClass* Class);
+	UE_API bool IsPackageSaved(FName PackageName);
 };
+
+#undef UE_API
diff --git a/Engine/Source/Editor/AIGraph/Public/AIGraphConnectionDrawingPolicy.h b/Engine/Source/Editor/AIGraph/Public/AIGraphConnectionDrawingPolicy.h
--- a/Engine/Source/Editor/AIGraph/Public/AIGraphConnectionDrawingPolicy.h
+++ b/Engine/Source/Editor/AIGraph/Public/AIGraphConnectionDrawingPolicy.h
@@ -11,6 +11,8 @@
 #include "Templates/SharedPointer.h"
 #include "Widgets/SWidget.h"
 
+#define UE_API AIGRAPH_API
+
 class FArrangedChildren;
 class FArrangedWidget;
 class FSlateRect;
@@ -21,24 +23,26 @@ class UEdGraphNode;
 class UEdGraphPin;
 struct FGeometry;
 
-class AIGRAPH_API FAIGraphConnectionDrawingPolicy : public FConnectionDrawingPolicy
+class FAIGraphConnectionDrawingPolicy : public FConnectionDrawingPolicy
 {
 protected:
 	UEdGraph* GraphObj;
 	TMap<UEdGraphNode*, int32> NodeWidgetMap;
 
 public:
-	FAIGraphConnectionDrawingPolicy(int32 InBackLayerID, int32 InFrontLayerID, float ZoomFactor, const FSlateRect& InClippingRect, FSlateWindowElementList& InDrawElements, UEdGraph* InGraphObj);
+	UE_API FAIGraphConnectionDrawingPolicy(int32 InBackLayerID, int32 InFrontLayerID, float ZoomFactor, const FSlateRect& InClippingRect, FSlateWindowElementList& InDrawElements, UEdGraph* InGraphObj);
 
 	// FConnectionDrawingPolicy interface 
-	virtual void DetermineWiringStyle(UEdGraphPin* OutputPin, UEdGraphPin* InputPin, /*inout*/ FConnectionParams& Params) override;
-	virtual void Draw(TMap<TSharedRef<SWidget>, FArrangedWidget>& PinGeometries, FArrangedChildren& ArrangedNodes) override;
-	virtual void DrawSplineWithArrow(const FGeometry& StartGeom, const FGeometry& EndGeom, const FConnectionParams& Params) override;
-	virtual void DrawSplineWithArrow(const FVector2f& StartPoint, const FVector2f& EndPoint, const FConnectionParams& Params) override;
-	virtual void DrawPreviewConnector(const FGeometry& PinGeometry, const FVector2f& StartPoint, const FVector2f& EndPoint, UEdGraphPin* Pin) override;
-	virtual FVector2f ComputeSplineTangent(const FVector2f& Start, const FVector2f& End) const override;
+	UE_API virtual void DetermineWiringStyle(UEdGraphPin* OutputPin, UEdGraphPin* InputPin, /*inout*/ FConnectionParams& Params) override;
+	UE_API virtual void Draw(TMap<TSharedRef<SWidget>, FArrangedWidget>& PinGeometries, FArrangedChildren& ArrangedNodes) override;
+	UE_API virtual void DrawSplineWithArrow(const FGeometry& StartGeom, const FGeometry& EndGeom, const FConnectionParams& Params) override;
+	UE_API virtual void DrawSplineWithArrow(const FVector2f& StartPoint, const FVector2f& EndPoint, const FConnectionParams& Params) override;
+	UE_API virtual void DrawPreviewConnector(const FGeometry& PinGeometry, const FVector2f& StartPoint, const FVector2f& EndPoint, UEdGraphPin* Pin) override;
+	UE_API virtual FVector2f ComputeSplineTangent(const FVector2f& Start, const FVector2f& End) const override;
 	// End of FConnectionDrawingPolicy interface
 
 protected:
-	void Internal_DrawLineWithArrow(const UE::Slate::FDeprecateVector2DParameter& StartAnchorPoint, const UE::Slate::FDeprecateVector2DParameter& EndAnchorPoint, const FConnectionParams& Params);
+	UE_API void Internal_DrawLineWithArrow(const UE::Slate::FDeprecateVector2DParameter& StartAnchorPoint, const UE::Slate::FDeprecateVector2DParameter& EndAnchorPoint, const FConnectionParams& Params);
 };
+
+#undef UE_API
